RelativeIKOp: Add correct module initialization files for anim info module.

diff --git a/Engine/Plugins/Experimental/Animation/RelativeIKOp/Source/RelativeBodyAnimInfo/Private/RelativeBodyAnimInfoModule.cpp b/Engine/Plugins/Experimental/Animation/RelativeIKOp/Source/RelativeBodyAnimInfo/Private/RelativeBodyAnimInfoModule.cpp
new file mode 100644
--- /dev/null
+++ b/Engine/Plugins/Experimental/Animation/RelativeIKOp/Source/RelativeBodyAnimInfo/Private/RelativeBodyAnimInfoModule.cpp
@@ -0,0 +1,17 @@
+// Copyright Epic Games, Inc. All Rights Reserved.
+
+#include "RelativeBodyAnimInfoModule.h"
+
+#define LOCTEXT_NAMESPACE "RelativeBodyAnimInfoModule"
+
+void FRelativeBodyAnimInfoModule::StartupModule()
+{
+}
+
+void FRelativeBodyAnimInfoModule::ShutdownModule()
+{
+}
+
+#undef LOCTEXT_NAMESPACE
+    
+IMPLEMENT_MODULE(FRelativeBodyAnimInfoModule, RelativeBodyAnimInfo)
\ No newline at end of file
diff --git a/Engine/Plugins/Experimental/Animation/RelativeIKOp/Source/RelativeBodyAnimInfo/Public/RelativeBodyAnimInfoModule.h b/Engine/Plugins/Experimental/Animation/RelativeIKOp/Source/RelativeBodyAnimInfo/Public/RelativeBodyAnimInfoModule.h
new file mode 100644
--- /dev/null
+++ b/Engine/Plugins/Experimental/Animation/RelativeIKOp/Source/RelativeBodyAnimInfo/Public/RelativeBodyAnimInfoModule.h
@@ -0,0 +1,14 @@
+// Copyright Epic Games, Inc. All Rights Reserved.
+
+#pragma once
+
+#include "Modules/ModuleManager.h"
+
+class FRelativeBodyAnimInfoModule : public IModuleInterface
+{
+public:
+
+	/** IModuleInterface implementation */
+	virtual void StartupModule() override;
+	virtual void ShutdownModule() override;
+};
