RelativeIKOp: Added missing revert implementation for modifier, changed track names so PhysicalIK can coexist, and added scoped task progress-bar for baking.

diff --git a/Engine/Plugins/Experimental/Animation/RelativeIKOp/Source/RelativeBodyAnimUtils/Private/RelativeBodyAnimModifier.cpp b/Engine/Plugins/Experimental/Animation/RelativeIKOp/Source/RelativeBodyAnimUtils/Private/RelativeBodyAnimModifier.cpp
--- a/Engine/Plugins/Experimental/Animation/RelativeIKOp/Source/RelativeBodyAnimUtils/Private/RelativeBodyAnimModifier.cpp
+++ b/Engine/Plugins/Experimental/Animation/RelativeIKOp/Source/RelativeBodyAnimUtils/Private/RelativeBodyAnimModifier.cpp
@@ -10,11 +10,14 @@
 #include "Animation/AnimSequence.h"
 #include "Components/SkeletalMeshComponent.h"
 #include "Engine/SkeletalMesh.h"
+#include "Misc/ScopedSlowTask.h"
 #include "PhysicsEngine/PhysicsAsset.h"
 #include "PhysicsEngine/SkeletalBodySetup.h"
 #include "Rendering/SkeletalMeshRenderData.h"
 #include "Spatial/PointHashGrid3.h"
 
+#define LOCTEXT_NAMESPACE "RelativeBodyAnimModifier"
+
 void URelativeBodyAnimModifier::OnApply_Implementation(UAnimSequence* InAnimation)
 {
 	// TODO: Make log category for relative body utils?
@@ -85,7 +88,7 @@ void URelativeBodyAnimModifier::OnApply_Implementation(UAnimSequence* InAnimatio
 		bool bUseDenseNotifyInfo = NotifyClass->IsChildOf(URelativeBodyBakeAnimNotify::StaticClass());
 		if (bUseDenseNotifyInfo)
 		{
-			FName TrackName("<" + SkeletalMeshAsset->GetName() + ">" + "DenseInfo");
+			FName TrackName("RelBody-<" + SkeletalMeshAsset->GetName() + ">");
 			UE_LOG(LogAnimation, Verbose, TEXT("TrackName %d: '%s'."), GeneratedNotifyTracks.Num()-1, *TrackName.ToString());
 			const bool bDoesTrackNameAlreadyExist = UAnimationBlueprintLibrary::IsValidAnimNotifyTrackName(InAnimation, TrackName);
 			if (!bDoesTrackNameAlreadyExist)
@@ -99,7 +102,7 @@ void URelativeBodyAnimModifier::OnApply_Implementation(UAnimSequence* InAnimatio
 			{
 				return;
 			}
-
+			
 			for (int32 BodyIndex1 = 0; BodyIndex1 < NumBodies; ++BodyIndex1)
 			{
 				if (BodyIndicesToIgnore[BodyIndex1])
@@ -152,10 +155,15 @@ void URelativeBodyAnimModifier::OnApply_Implementation(UAnimSequence* InAnimatio
 			BakedAnimNotifyInfo->NumSamples = SampleNum;
 		}
 
+		FScopedSlowTask BakePairVertsTask((float)SampleNum, LOCTEXT("BakeVertsTaskText", "Baking Body Pair Verts..."));
+		BakePairVertsTask.MakeDialog();
+
 		// Get ground levels and max speed values.
 		TArray<FVector3f> VLocations;
 		for (int SampleIndex = 0; SampleIndex < SampleNum; ++SampleIndex)
 		{
+			BakePairVertsTask.EnterProgressFrame(1.0f);
+			
 			const float SampleTime = FMath::Clamp(static_cast<float>(SampleIndex) * SampleStep, 0.0f, SequenceLength);
 			const float FutureSampleTime = FMath::Clamp((static_cast<float>(SampleIndex) + 1.0f) * SampleStep, 0.0f, SequenceLength);
 
@@ -295,7 +303,13 @@ void URelativeBodyAnimModifier::OnApply_Implementation(UAnimSequence* InAnimatio
 
 void URelativeBodyAnimModifier::OnRevert_Implementation(UAnimSequence* InAnimation)
 {
-	
+	// Delete any generate tracks.
+	for (const FName GeneratedNotifyTrack : GeneratedNotifyTracks)
+	{
+		UAnimationBlueprintLibrary::RemoveAnimationNotifyTrack(InAnimation, GeneratedNotifyTrack);
+	}
+
+	GeneratedNotifyTracks.Reset();
 }
 
 UPhysicsAsset* URelativeBodyAnimModifier::GetPhysicsAsset() const
@@ -491,3 +505,5 @@ void URelativeBodyAnimModifier::GetSkinnedVertices(TArray<FVector3f>& VLocations
 		}
 	}
 }
+
+#undef LOCTEXT_NAMESPACE
diff --git a/Engine/Plugins/Experimental/Animation/RelativeIKOp/Source/RelativeBodyAnimUtils/Public/RelativeBodyAnimModifier.h b/Engine/Plugins/Experimental/Animation/RelativeIKOp/Source/RelativeBodyAnimUtils/Public/RelativeBodyAnimModifier.h
--- a/Engine/Plugins/Experimental/Animation/RelativeIKOp/Source/RelativeBodyAnimUtils/Public/RelativeBodyAnimModifier.h
+++ b/Engine/Plugins/Experimental/Animation/RelativeIKOp/Source/RelativeBodyAnimUtils/Public/RelativeBodyAnimModifier.h
@@ -89,10 +89,6 @@ private:
 	UPROPERTY(VisibleAnywhere, AdvancedDisplay, meta = (Category = "Settings"))
 	TSet<FName> GeneratedNotifyTracks;
 
-	/** Keep track of tracks modified during modifier application */
-	UPROPERTY(VisibleAnywhere, AdvancedDisplay, meta = (Category = "Settings"))
-	TSet<FName> ProcessedNotifyTracks;
-
 	UPhysicsAsset* GetPhysicsAsset() const;
 	
 	void CacheBodyDataForSourceMesh(FRelativeBodySourceData& OutSourceData);
