RelativeIKOp: Code cleanup, debug draw fixes, and added preliminary source/target body mapping

diff --git a/Engine/Plugins/Experimental/Animation/RelativeIKOp/Source/RelativeIKOp/Private/RelativeIKOp.cpp b/Engine/Plugins/Experimental/Animation/RelativeIKOp/Source/RelativeIKOp/Private/RelativeIKOp.cpp
--- a/Engine/Plugins/Experimental/Animation/RelativeIKOp/Source/RelativeIKOp/Private/RelativeIKOp.cpp
+++ b/Engine/Plugins/Experimental/Animation/RelativeIKOp/Source/RelativeIKOp/Private/RelativeIKOp.cpp
@@ -87,6 +87,7 @@ void FIKRetargetRelativeIKOp::Run(
 	const TArray<FTransform>& InSourceGlobalPose,
 	TArray<FTransform>& OutTargetGlobalPose)
 {
+	// TODO: Reset/clear debug info structs here to
 	if (InProcessor.IsIKForcedOff())
 	{
 		return;
@@ -109,52 +110,66 @@ void FIKRetargetRelativeIKOp::Run(
 	int32 SampleIdx = FMath::Min(static_cast<int32>(NotifyTime * SampleRate), CachedNotifyInfo->NumSamples-1);
 	int32 FramePairIdx = NumBodies * SampleIdx;
 
+	FDebugRelativeIKDrawInfo LocalDebugInfo;
+	
+#if WITH_EDITOR
+	LocalDebugInfo.SourcePairVerts.Reserve(CachedNotifyInfo->bBodyPairsIsParentDominates.Num());
+	LocalDebugInfo.TargetPairVerts.Reserve(CachedNotifyInfo->bBodyPairsIsParentDominates.Num());
+
+	TArray<FDebugRelativeTargetPairSpace> LocalTargetPairSpaces;
+	LocalTargetPairSpaces.SetNum(CachedNotifyInfo->BodyPairs.Num());
+#endif //WITH_EDITOR
+
 	const double SourceScale = InProcessor.GetSourceScaleFactor();
 	const double InvSourceScale = (Settings.bIgnoreSourceScale) ? 1.0 / InProcessor.GetSourceScaleFactor() : 1.0;
 
 	const FTransform SourceScaleTfm(FQuat::Identity, FVector::Zero(), FVector(SourceScale));
-	const FTransform InvScaleTfm(FQuat::Identity, FVector::Zero(), FVector(InvSourceScale)); 
-
-	TArray<TTuple<FVector,FVector>> DebugSourceVerts;
-	TArray<TTuple<FVector,FVector>> DebugTargetVerts;
-
-	DebugSourceVerts.Reserve(CachedNotifyInfo->bBodyPairsIsParentDominates.Num());
-	DebugTargetVerts.Reserve(CachedNotifyInfo->bBodyPairsIsParentDominates.Num());
-
-	TArray<TTuple<FVector,FVector,FVector,FVector,FVector>> DebugTargetPairVerts;
-	DebugTargetPairVerts.SetNum(CachedNotifyInfo->BodyPairs.Num());
+	const FTransform InvScaleTfm(FQuat::Identity, FVector::Zero(), FVector(InvSourceScale));
 	
 	FIKRigGoalContainer& GoalContainer = InProcessor.GetIKRigGoalContainer();
 	TArray<FIKRigGoal>& IKGoals = GoalContainer.GetGoalArray();
 
-	TMap<FName,double> ChainLengthMap;
-	TMap<FName,FVector> ChainStartMap;
+	CacheChainLengthMap.Reset();
+	CacheChainStartMap.Reset();
 	
 	const double RetargetContactAlpha = FMath::Clamp(Settings.RetargetContactAlpha, 0.0, 1.0);
 	const double RetargetSpringAlpha = FMath::Clamp(Settings.RetargetSpringAlpha, 0.0, 1.0);
 	// Get pair target body verts and weights
 	for (int PairIdx = 0; PairIdx < CachedNotifyInfo->bBodyPairsIsParentDominates.Num(); ++PairIdx)
 	{
+		// TODO: Allow selecting a subset of body pairs in retargeter
 		FrameBoneWeights[2*PairIdx] = 0;
 		FrameBoneWeights[2*PairIdx + 1] = 0;
 		
-		FName PrimaryBone = CachedNotifyInfo->BodyPairs[2*PairIdx];
-		FName OtherBone = CachedNotifyInfo->BodyPairs[2*PairIdx + 1];
+		FName SourcePrimaryBone = CachedNotifyInfo->BodyPairs[2*PairIdx];
+		FName SourceOtherBone = CachedNotifyInfo->BodyPairs[2*PairIdx + 1];
+
+		FName TargetPrimaryBone = ApplyBodyMap(SourcePrimaryBone);
+		FName TargetOtherBone = ApplyBodyMap(SourceOtherBone);
 
 		// TODO: Early-out if missing body in current source/target PhysAsset
 		// Check this in update?
+		if (TargetPrimaryBone == NAME_None || TargetOtherBone == NAME_None)
+		{
+			continue;
+		}
 
-		int32 SourcePrimaryPoseIdx = CacheSourceSkelIndices[PrimaryBone];
-		int32 SourceOtherPoseIdx = CacheSourceSkelIndices[OtherBone];
+		if (!CacheSourceSkelIndices.Contains(SourcePrimaryBone) || !CacheSourceSkelIndices.Contains(SourceOtherBone))
+		{
+			continue;
+		}
+
+		int32 SourcePrimaryPoseIdx = CacheSourceSkelIndices[SourcePrimaryBone];
+		int32 SourceOtherPoseIdx = CacheSourceSkelIndices[SourceOtherBone];
 
 		// Compute source transforms and source-space vertex distance
 		const FVector3f& LocalPrimaryVert = CachedNotifyInfo->BodyPairsLocalReference[FramePairIdx + 2*PairIdx];
-		FTransform SourcePrimaryBodyTfm = GetBodyTransform(SourcePhysicsAsset, PrimaryBone) * SourceScaleTfm;
+		FTransform SourcePrimaryBodyTfm = GetBodyTransform(SourcePhysicsAsset, SourcePrimaryBone) * SourceScaleTfm;
 		FTransform SourcePrimaryCompTfm = SourcePrimaryBodyTfm * InSourceGlobalPose[SourcePrimaryPoseIdx] * InvScaleTfm;
 		FVector SourcePrimaryCompV = SourcePrimaryCompTfm.TransformPosition(FVector(LocalPrimaryVert));
 
 		const FVector3f& LocalOtherVert = CachedNotifyInfo->BodyPairsLocalReference[FramePairIdx + 2*PairIdx + 1];
-		FTransform SourceOtherBodyTfm = GetBodyTransform(SourcePhysicsAsset, OtherBone) * SourceScaleTfm;
+		FTransform SourceOtherBodyTfm = GetBodyTransform(SourcePhysicsAsset, SourceOtherBone) * SourceScaleTfm;
 		FTransform SourceOtherCompTfm = SourceOtherBodyTfm * InSourceGlobalPose[SourceOtherPoseIdx] * InvScaleTfm;
 		FVector SourceOtherCompV = SourceOtherCompTfm.TransformPosition(FVector(LocalOtherVert));
 
@@ -167,10 +182,10 @@ void FIKRetargetRelativeIKOp::Run(
 		FVector SourceOtherPrimaryV = SourcePrimaryCompTfm.InverseTransformPosition(SourceOtherCompV);
 
 		// Get average of retarget verts
-		int32 TargetPrimaryPoseIdx = CacheTargetSkelIndices[PrimaryBone];
-		int32 TargetOtherPoseIdx = CacheTargetSkelIndices[OtherBone];
-		FTransform TargetPrimaryBodyTfm = GetBodyTransform(TargetPhysicsAsset, PrimaryBone);
-		FTransform TargetOtherBodyTfm = GetBodyTransform(TargetPhysicsAsset, OtherBone);
+		int32 TargetPrimaryPoseIdx = CacheTargetSkelIndices[TargetPrimaryBone];
+		int32 TargetOtherPoseIdx = CacheTargetSkelIndices[TargetOtherBone];
+		FTransform TargetPrimaryBodyTfm = GetBodyTransform(TargetPhysicsAsset, TargetPrimaryBone);
+		FTransform TargetOtherBodyTfm = GetBodyTransform(TargetPhysicsAsset, TargetOtherBone);
 
 		FTransform TargetPrimaryCompTfm = TargetPrimaryBodyTfm * OutTargetGlobalPose[TargetPrimaryPoseIdx];
 		FTransform TargetOtherCompTfm = TargetOtherBodyTfm * OutTargetGlobalPose[TargetOtherPoseIdx];
@@ -190,56 +205,19 @@ void FIKRetargetRelativeIKOp::Run(
 
 		// double TargetVDist = FVector::Distance(TargetRelativeToPrimaryV, TargetRelativeToOtherV);
 		// double TargetDRat = FMath::Min(1.0, 2.0 * FMath::Min(SourceVDist/TargetVDist, TargetVDist/SourceVDist));
-
+		
 		FVector RelativePrimaryOffset = TargetPrimaryCompV_P - TargetRelativeToPrimaryV;
 		FVector RelativeOtherOffset = TargetOtherCompV_O - TargetRelativeToOtherV;
 
 		// TODO: Should Eventually remove testing setup (keep dist weighting or drop it)
 		bool bDistanceWeighting = CachedNotifyInfo->bBodyPairsIsParentDominates[PairIdx] && Settings.bTestDistContactAlpha;
-
-		// TODO: Experiment with bi-directional distance weighting when Primary/Other can both move
+		
 		FDoubleInterval FeasibleRange(0.0, 1.0);
-		if (bDistanceWeighting && Settings.bTestFeasibilityWeight)
+		UpdateCacheChainInfo(OutTargetGlobalPose, TargetOtherBone);
+		if (bDistanceWeighting)
 		{
-			// Figure out alpha of feasible point along contribution target line 
-			if (CacheBoneChains.Contains(OtherBone) && !ChainLengthMap.Contains(OtherBone))
-			{
-				TArray<FTransform> ChainTransforms;
-				ChainTransforms.Reserve(CacheBoneChains[OtherBone]->BoneIndices.Num());
-				for (int32 BoneIdx : CacheBoneChains[OtherBone]->BoneIndices)
-				{
-					ChainTransforms.Add(OutTargetGlobalPose[BoneIdx]);
-				}
-
-				ChainLengthMap.Emplace(OtherBone, CacheBoneChains[OtherBone]->GetChainLength(ChainTransforms));
-				ChainStartMap.Emplace(OtherBone, OutTargetGlobalPose[CacheBoneChains[OtherBone]->BoneIndices[0]].GetTranslation());
-			}
-			
-			FVector OffsetChainStart = ChainStartMap[OtherBone] - RelativeOtherOffset - TargetOtherBoneOffset;
-			// TODO: Better bias range
-			double CheckRadius = FMath::Max(ChainLengthMap[OtherBone] + Settings.FeasibilityLengthBias, 0.0);
-			
-			double ChainDistStart = FVector::Distance(TargetRelativeToOtherV, OffsetChainStart);
-			double ChainDistEnd = FVector::Distance(TargetRelativeToPrimaryV, OffsetChainStart);
-
-			bool UpperRange = ChainDistStart > ChainDistEnd; 
-			if (FMath::Max(ChainDistStart, ChainDistEnd) <= CheckRadius)
-			{
-				FeasibleRange = FDoubleInterval(0.0,1.0);
-			}
-			else if (FMath::Min(ChainDistStart, ChainDistEnd) >= CheckRadius)
-			{
-				double RangePt = (UpperRange) ? 1.0 : 0.0;
-				FeasibleRange = FDoubleInterval(RangePt,RangePt);
-			}
-			else
-			{
-				double RangePt = CalcLineSphereIntersect(OffsetChainStart, CheckRadius, TargetRelativeToOtherV, TargetRelativeToPrimaryV);
-				double MinVal = (UpperRange) ? RangePt : 0.0;
-				double MaxVal = (UpperRange) ? 1.0 : RangePt;
-
-				FeasibleRange = FDoubleInterval(MinVal, MaxVal);
-			}
+			// TODO: Experiment with bi-directional distance weighting when Primary/Other can both move
+			ComputeFeasibilityRange(FeasibleRange, TargetOtherBone, RelativeOtherOffset + TargetOtherBoneOffset, TargetRelativeToOtherV, TargetRelativeToPrimaryV);
 		}
 
 		double FeasibilityWeight = FMath::Lerp(FeasibleRange.Min, FeasibleRange.Max, LinearWeight);
@@ -259,41 +237,33 @@ void FIKRetargetRelativeIKOp::Run(
 		// Debug draw verts
 		if (Settings.bEnableDebugDraw)
 		{
-			DebugSourceVerts.Add(TTuple<FVector,FVector>(SourcePrimaryCompV, SourceOtherCompV));
-
-			if (Settings.DrawTargetPairBones.Contains(PrimaryBone))
+			LocalDebugInfo.SourcePairVerts.Add(TTuple<FVector,FVector>(SourcePrimaryCompV, SourceOtherCompV));
+			if (Settings.DrawTargetPairBones.Contains(TargetPrimaryBone))
 			{
-				DebugTargetVerts.Add(TTuple<FVector,FVector>(TargetPrimaryCompV_P, TargetOtherCompV_O));
-				DebugTargetVerts.Add(TTuple<FVector,FVector>(TargetPrimaryCompV_O, TargetOtherCompV_P));
+				LocalDebugInfo.TargetPairVerts.Add(TTuple<FVector,FVector>(TargetPrimaryCompV_P, TargetOtherCompV_O));
+				LocalDebugInfo.TargetPairVerts.Add(TTuple<FVector,FVector>(TargetPrimaryCompV_O, TargetOtherCompV_P));
 			}
-
-			FVector DbgFeasTargetPrimaryMin = FMath::Lerp(TargetRelativeToOtherV, TargetRelativeToPrimaryV, FeasibleRange.Min) + RelativePrimaryOffset;
-			FVector DbgFeasTargetPrimaryMax = FMath::Lerp(TargetRelativeToOtherV, TargetRelativeToPrimaryV, FeasibleRange.Max) + RelativePrimaryOffset;
-			FVector DbgFeasTargetOtherMin = FMath::Lerp(TargetRelativeToOtherV, TargetRelativeToPrimaryV, FeasibleRange.Min) + RelativeOtherOffset;
-			FVector DbgFeasTargetOtherMax = FMath::Lerp(TargetRelativeToOtherV, TargetRelativeToPrimaryV, FeasibleRange.Max) + RelativeOtherOffset;
-		
-			DebugTargetPairVerts[2*PairIdx] = TTuple<FVector,FVector,FVector,FVector,FVector>(TargetRelativeToPrimaryV+RelativePrimaryOffset, TargetRelativeToOtherV+RelativePrimaryOffset, WeightedTargetPrimary, DbgFeasTargetPrimaryMin, DbgFeasTargetPrimaryMax);
-			DebugTargetPairVerts[2*PairIdx+1] = TTuple<FVector,FVector,FVector,FVector,FVector>(TargetRelativeToPrimaryV+RelativeOtherOffset, TargetRelativeToOtherV+RelativeOtherOffset, WeightedTargetOther, DbgFeasTargetOtherMin, DbgFeasTargetOtherMax);
+			
+			LocalTargetPairSpaces[2*PairIdx] = FDebugRelativeTargetPairSpace{WeightedTargetPrimary, TargetRelativeToOtherV+RelativePrimaryOffset, TargetRelativeToPrimaryV+RelativePrimaryOffset, FeasibleRange};
+			LocalTargetPairSpaces[2*PairIdx+1] = FDebugRelativeTargetPairSpace{WeightedTargetOther, TargetRelativeToOtherV+RelativeOtherOffset, TargetRelativeToPrimaryV+RelativeOtherOffset, FeasibleRange};
 		}
 #endif //WITH_EDITOR
 	}
-
+	
+	LocalDebugInfo.TargetGoals.Reserve(CacheSourceBodyEffectVertIdx.Num());
 	// Update bone goals to relative IK positions
-	TArray<TTuple<FVector,double>> DebugTestSpheres;
-	TArray<TTuple<FVector,TArray<FVector>>> DebugIKGoalPoints;
-	TArray<TTuple<FVector,FVector,FVector,FVector,FVector>> DebugTargetContrib;
-	DebugIKGoalPoints.Reserve(CacheBodyEffectVertIdx.Num());
-	for (int EffectIdx = 0; EffectIdx < CacheBodyEffectVertIdx.Num(); ++ EffectIdx)
+	for (int EffectIdx = 0; EffectIdx < CacheSourceBodyEffectVertIdx.Num(); ++ EffectIdx)
 	{
-		TArray<int32>& EffectIndices = CacheBodyEffectVertIdx[EffectIdx];
+		TArray<int32>& EffectIndices = CacheSourceBodyEffectVertIdx[EffectIdx];
 		if (EffectIndices.IsEmpty())
 		{
 			continue;
 		}
 
-		FName BoneName = CacheEffectBones[EffectIdx];
+		FName SourceBoneName = CacheSourceEffectBones[EffectIdx];
+		FName TargetBoneName = ApplyBodyMap(SourceBoneName);
 		// HACK: Quick and dirty goal-by-bone find
-		FIKRigGoal* IKGoal = IKGoals.FindByPredicate([BoneName](const FIKRigGoal& Other){return Other.BoneName == BoneName;});
+		FIKRigGoal* IKGoal = IKGoals.FindByPredicate([TargetBoneName](const FIKRigGoal& Other){return Other.BoneName == TargetBoneName;});
 		if (!IKGoal)
 		{
 			continue;
@@ -301,7 +271,14 @@ void FIKRetargetRelativeIKOp::Run(
 
 		double TotalWeight = 0;
 		FVector GoalPoint = FVector::ZeroVector;
-		DebugIKGoalPoints.AddDefaulted();
+
+#if WITH_EDITOR
+		if (Settings.bEnableDebugDraw)
+		{
+			LocalDebugInfo.TargetGoals.AddDefaulted();
+			LocalDebugInfo.TargetGoals.Last().Goal = FVector::ZeroVector;
+		}
+#endif //WITH_EDITOR
 
 		double MaxWeight = 0;
 		for (int Index : EffectIndices)
@@ -317,16 +294,22 @@ void FIKRetargetRelativeIKOp::Run(
 #if WITH_EDITOR
 			if (Settings.bEnableDebugDraw)
 			{
-				DebugIKGoalPoints.Last().Get<1>().Add(FrameBoneTargets[Index]);
-				DebugTargetContrib.Emplace(DebugTargetPairVerts[Index]);
+				LocalDebugInfo.TargetGoals.Last().PairTargets.Add(FrameBoneTargets[Index]);
+				LocalDebugInfo.PairContrib.Emplace(LocalTargetPairSpaces[Index]);
 			}
 #endif //WITH_EDITOR
 		}
 
+		
 		if (!FMath::IsNearlyZero(TotalWeight))
 		{
 			GoalPoint /= TotalWeight;
-			DebugIKGoalPoints.Last().Get<0>() = GoalPoint;
+#if WITH_EDITOR
+			if (Settings.bEnableDebugDraw)
+			{
+				LocalDebugInfo.TargetGoals.Last().Goal = GoalPoint;
+			}
+#endif // WITH_EDITOR
 			
 			// TODO: Play with total weight functions to alpha goal
 			float GoalAlpha = FMath::Min(1.0f, static_cast<float>(TotalWeight*FMath::Clamp(Settings.ContributionSumWeight, 0.0f, 1.0f)));
@@ -338,7 +321,7 @@ void FIKRetargetRelativeIKOp::Run(
 			// Sphere-capsule intersection goal deltas
 			if (Settings.bTestIntersectionFixup)
 			{
-				FKShapeElem* BoneShape = FindBodyShape(TargetPhysicsAsset, BoneName);
+				FKShapeElem* BoneShape = FindBodyShape(TargetPhysicsAsset, TargetBoneName);
 				double BoneRadius = CalcShapeSmallRadius(BoneShape);
 				if (FMath::IsNearlyZero(BoneRadius))
 				{
@@ -347,7 +330,7 @@ void FIKRetargetRelativeIKOp::Run(
 				
 				for (FName BodyName : Settings.IntersectBodies)
 				{
-					if (BoneName == BodyName)
+					if (TargetBoneName == BodyName)
 					{
 						continue;
 					}
@@ -360,6 +343,10 @@ void FIKRetargetRelativeIKOp::Run(
 					int32 BodyPoseIdx = CacheTargetSkelIndices[BodyName];
 					FTransform BodyTfm = OutTargetGlobalPose[BodyPoseIdx];
 					FKShapeElem* ShapeElem = FindBodyShape(TargetPhysicsAsset, BodyName);
+					if (!ShapeElem)
+					{
+						continue;
+					}
 
 					FVector DeltaDir = FVector::ZeroVector;
 					double Dist = CalcIntersectionDelta(GoalPoint, BoneRadius, BodyTfm, ShapeElem, DeltaDir);
@@ -369,10 +356,12 @@ void FIKRetargetRelativeIKOp::Run(
 					}
 				}
 
-				if (Settings.bDebugDrawIntersectTest)
+#if WITH_EDITOR
+				if (Settings.bEnableDebugDraw)
 				{
-					DebugTestSpheres.Emplace(IKGoal->Position, BoneRadius);
+					LocalDebugInfo.TestSpheres.Emplace(IKGoal->Position, BoneRadius);
 				}
+#endif //WITH_EDITOR
 			}
 		}
 	}
@@ -381,28 +370,35 @@ void FIKRetargetRelativeIKOp::Run(
 	if (Settings.bEnableDebugDraw)
 	{
 		FScopeLock ScopeLock(&DebugDataMutex);
-		DebugDrawSourcePairVerts = DebugSourceVerts;
-		DebugDrawTargetPairVerts = DebugTargetVerts;
 
-		DebugDrawTargetContrib = DebugTargetContrib;
-		DebugDrawGoals = DebugIKGoalPoints;
+		DebugDrawInfo.TargetGoals = LocalDebugInfo.TargetGoals;
+		DebugDrawInfo.PairContrib = LocalDebugInfo.PairContrib;
+		DebugDrawInfo.SourcePairVerts = LocalDebugInfo.SourcePairVerts;
+		DebugDrawInfo.TargetPairVerts = LocalDebugInfo.TargetPairVerts;
+		DebugDrawInfo.TestSpheres = LocalDebugInfo.TestSpheres;
 
-		DebugDrawSourceBodyTfms.Reset(CacheEffectBones.Num());
-		DebugDrawTargetBodyTfms.Reset(CacheEffectBones.Num());
-		for (FName BoneName : CacheEffectBones)
+		DebugDrawInfo.SourceBodyTfms.Reset(CacheSourceEffectBones.Num());
+		DebugDrawInfo.TargetBodyTfms.Reset(CacheSourceEffectBones.Num());
+		for (FName SourceBoneName : CacheSourceEffectBones)
 		{
-			int32 SourcePoseIdx = CacheSourceSkelIndices[BoneName];
-			FTransform SourceBodyTfm = GetBodyTransform(SourcePhysicsAsset, BoneName) * SourceScaleTfm;
+			FName TargetBoneName = ApplyBodyMap(SourceBoneName);
+			if (!CacheSourceSkelIndices.Contains(SourceBoneName) || !CacheTargetSkelIndices.Contains(TargetBoneName))
+			{
+				continue;
+			}
+			
+			int32 SourcePoseIdx = CacheSourceSkelIndices[SourceBoneName];
+			FTransform SourceBodyTfm = GetBodyTransform(SourcePhysicsAsset, SourceBoneName) * SourceScaleTfm;
 			FTransform SourceBodyCompTfm = SourceBodyTfm * InSourceGlobalPose[SourcePoseIdx];
-			DebugDrawSourceBodyTfms.Add(SourceBodyCompTfm);
-
-			int32 TargetPoseIdx = CacheTargetSkelIndices[BoneName];
-			FTransform TargetBodyTfm = GetBodyTransform(TargetPhysicsAsset, BoneName);
+			DebugDrawInfo.SourceBodyTfms.Add(SourceBodyCompTfm);
+			
+			int32 TargetPoseIdx = CacheTargetSkelIndices[TargetBoneName];
+			FTransform TargetBodyTfm = GetBodyTransform(TargetPhysicsAsset, TargetBoneName);
 			FTransform TargetBodyCompTfm = TargetBodyTfm * OutTargetGlobalPose[TargetPoseIdx];
-			DebugDrawTargetBodyTfms.Add(TargetBodyCompTfm);
+			DebugDrawInfo.TargetBodyTfms.Add(TargetBodyCompTfm);
 		}
 
-		DebugDrawIntersectTfms.Reset(Settings.IntersectBodies.Num());
+		DebugDrawInfo.TargetIntersectTfms.Reset(Settings.IntersectBodies.Num());
 		for (FName BodyName : Settings.IntersectBodies)
 		{
 			if (!CacheTargetSkelIndices.Contains(BodyName))
@@ -411,12 +407,95 @@ void FIKRetargetRelativeIKOp::Run(
 			}
 			
 			int32 TargetPoseIdx = CacheTargetSkelIndices[BodyName];
-			DebugDrawIntersectTfms.Add(OutTargetGlobalPose[TargetPoseIdx]);
+			DebugDrawInfo.TargetIntersectTfms.Add(OutTargetGlobalPose[TargetPoseIdx]);
 		}
 	}
 #endif
 }
 
+FName FIKRetargetRelativeIKOp::ApplyBodyMap(FName BodyName)
+{
+	if (Settings.BodyMapping.IsEmpty())
+	{
+		return BodyName;
+	}
+
+	// TODO: Debug-only logging of issues
+	if (!Settings.BodyMapping.Contains(BodyName))
+	{
+		// UE_LOG(LogAnimation, Warning, TEXT("Body name in baked data but not in mapping: %s"), *BodyName.ToString());
+		return NAME_None;
+	}
+
+	return Settings.BodyMapping[BodyName];
+}
+
+void FIKRetargetRelativeIKOp::UpdateCacheChainInfo(const TArray<FTransform>& TargetPose, FName BoneName)
+{
+	if (!CacheBoneChains.Contains(BoneName))
+	{
+		return;
+	}
+
+	if (CacheChainLengthMap.Contains(BoneName))
+	{
+		return;
+	}
+
+	TArray<FTransform> ChainTransforms;
+	ChainTransforms.Reserve(CacheBoneChains[BoneName]->BoneIndices.Num());
+	for (int32 BoneIdx : CacheBoneChains[BoneName]->BoneIndices)
+	{
+		ChainTransforms.Add(TargetPose[BoneIdx]);
+	}
+
+	CacheChainLengthMap.Emplace(BoneName, CacheBoneChains[BoneName]->GetChainLength(ChainTransforms));
+	CacheChainStartMap.Emplace(BoneName, TargetPose[CacheBoneChains[BoneName]->BoneIndices[0]].GetTranslation());
+}
+
+void FIKRetargetRelativeIKOp::ComputeFeasibilityRange(FDoubleInterval& OutFeasibleRange, FName BoneName, FVector BoneChainOffset, FVector TargetRelativeToBoneV, FVector TargetRelativeToOpposeBoneV)
+{
+	OutFeasibleRange = {0.0, 1.0};
+	if (!Settings.bTestFeasibilityWeight)
+	{
+		return;
+	}
+
+	if (!CacheChainLengthMap.Contains(BoneName))
+	{
+		return;
+	}
+	
+			
+	// FVector OffsetChainStart = ChainStartMap[OtherBone] - RelativeOtherOffset - TargetOtherBoneOffset;
+	FVector OffsetChainStart = CacheChainStartMap[BoneName] - BoneChainOffset;
+	// TODO: Better bias range
+	double CheckRadius = FMath::Max(CacheChainLengthMap[BoneName] + Settings.FeasibilityLengthBias, 0.0);
+			
+	double ChainDistStart = FVector::Distance(TargetRelativeToBoneV, OffsetChainStart);
+	double ChainDistEnd = FVector::Distance(TargetRelativeToOpposeBoneV, OffsetChainStart);
+
+	bool UpperRange = ChainDistStart > ChainDistEnd; 
+	if (FMath::Max(ChainDistStart, ChainDistEnd) <= CheckRadius)
+	{
+		OutFeasibleRange = FDoubleInterval(0.0,1.0);
+	}
+	else if (FMath::Min(ChainDistStart, ChainDistEnd) >= CheckRadius)
+	{
+		double RangePt = (UpperRange) ? 1.0 : 0.0;
+		OutFeasibleRange = FDoubleInterval(RangePt,RangePt);
+	}
+	else
+	{
+		// Figure out alpha range of feasible points along contribution target line
+		double RangePt = CalcLineSphereIntersect(OffsetChainStart, CheckRadius, TargetRelativeToBoneV, TargetRelativeToOpposeBoneV);
+		double MinVal = (UpperRange) ? RangePt : 0.0;
+		double MaxVal = (UpperRange) ? 1.0 : RangePt;
+
+		OutFeasibleRange = FDoubleInterval(MinVal, MaxVal);
+	}
+}
+
 double FIKRetargetRelativeIKOp::CalcLineSphereIntersect(const FVector& Center, double Radius, const FVector& StartPoint, const FVector& EndPoint)
 {
 	// NOTE: This function assumes two intersection points and returns positive ratio along line from start point
@@ -629,18 +708,32 @@ void FIKRetargetRelativeIKOp::UpdateCacheBoneChains(const FIKRetargetProcessor&
 	}
 }
 
-void FIKRetargetRelativeIKOp::UpdateBoneMaps(FName BoneName)
+void FIKRetargetRelativeIKOp::UpdateAllBoneMaps(FName SourceBoneName)
 {
-	int32 SourceIdx = SourceBoneNames.Find(BoneName);
-	int32 TargetIdx = TargetBoneNames.Find(BoneName);
+	FName TargetBoneName = ApplyBodyMap(SourceBoneName);
+	
+	int32 SourceIdx = SourceBoneNames.Find(SourceBoneName);
+	int32 TargetIdx = TargetBoneNames.Find(TargetBoneName);
 
 	if (SourceIdx == INDEX_NONE || TargetIdx == INDEX_NONE)
 	{
 		return;
 	}
 
-	CacheSourceSkelIndices.Emplace(BoneName, SourceIdx);
-	CacheTargetSkelIndices.Emplace(BoneName, TargetIdx);
+	CacheSourceSkelIndices.Emplace(SourceBoneName, SourceIdx);
+	CacheTargetSkelIndices.Emplace(TargetBoneName, TargetIdx);
+}
+
+void FIKRetargetRelativeIKOp::UpdateTargetBoneMap(FName TargetBoneName)
+{
+	int32 TargetIdx = TargetBoneNames.Find(TargetBoneName);
+
+	if (TargetIdx == INDEX_NONE)
+	{
+		return;
+	}
+	
+	CacheTargetSkelIndices.Emplace(TargetBoneName, TargetIdx);
 }
 
 void FIKRetargetRelativeIKOp::UpdateCacheSkelInfo(const FRetargetSkeleton& InSourceSkeleton, const FTargetSkeleton& InTargetSkeleton)
@@ -648,19 +741,19 @@ void FIKRetargetRelativeIKOp::UpdateCacheSkelInfo(const FRetargetSkeleton& InSou
 	SourceBoneNames = InSourceSkeleton.BoneNames;
 	TargetBoneNames = InTargetSkeleton.BoneNames;
 	
-	for (FName BodyName : Settings.DomainBodies)
+	for (FName SourceBodyName : Settings.DomainBodies)
 	{
-		UpdateBoneMaps(BodyName);
+		UpdateAllBoneMaps(SourceBodyName);
 	}
 
-	for (FName BodyName : Settings.ContactBodies)
+	for (FName SourceBodyName : Settings.ContactBodies)
 	{
-		UpdateBoneMaps(BodyName);
+		UpdateAllBoneMaps(SourceBodyName);
 	}
 
-	for (FName BodyName : Settings.IntersectBodies)
+	for (FName TargetBodyName : Settings.IntersectBodies)
 	{
-		UpdateBoneMaps(BodyName);
+		UpdateTargetBoneMap(TargetBodyName);
 	}
 }
 
@@ -705,8 +798,8 @@ void FIKRetargetRelativeIKOp::UpdateCacheNotifyInfo(URelativeBodyBakeAnimNotify*
 	CachedNotifyInfo = nullptr;
 
 	// Clear all cached data on update
-	CacheBodyEffectVertIdx.Reset();
-	CacheEffectBones.Reset();
+	CacheSourceBodyEffectVertIdx.Reset();
+	CacheSourceEffectBones.Reset();
 	FrameBoneTargets.Reset();
 	FrameBoneWeights.Reset();
 
@@ -725,20 +818,20 @@ void FIKRetargetRelativeIKOp::UpdateCacheNotifyInfo(URelativeBodyBakeAnimNotify*
 
 	SampleRate = FMath::RoundToFloat(1.0f / (CachedNotifyInfo->BodyPairsSampleTime[1] - CachedNotifyInfo->BodyPairsSampleTime[0]));
 
-	TMap<FName,int32> CheckBonesMap;
+	TMap<FName,int32> CheckSourceBonesMap;
 	for (int BodyIdx = 0; BodyIdx < NumBodies; ++BodyIdx)
 	{
-		FName BodyBone = CachedNotifyInfo->BodyPairs[BodyIdx];
-		if (CheckBonesMap.Contains(BodyBone))
+		FName SourceBodyBone = CachedNotifyInfo->BodyPairs[BodyIdx];
+		if (CheckSourceBonesMap.Contains(SourceBodyBone))
 		{
 			continue;
 		}
 
-		CheckBonesMap.Emplace(BodyBone, CacheEffectBones.Num());
-		CacheEffectBones.Add(BodyBone);
-		CacheBodyEffectVertIdx.AddDefaulted();
+		CheckSourceBonesMap.Emplace(SourceBodyBone, CacheSourceEffectBones.Num());
+		CacheSourceEffectBones.Add(SourceBodyBone);
+		CacheSourceBodyEffectVertIdx.AddDefaulted();
 
-		UpdateBoneMaps(BodyBone);
+		UpdateAllBoneMaps(SourceBodyBone);
 	}
 
 	// Cache bone-bone effectors for creating weighted goals
@@ -749,15 +842,15 @@ void FIKRetargetRelativeIKOp::UpdateCacheNotifyInfo(URelativeBodyBakeAnimNotify*
 		FName PrimaryBone = CachedNotifyInfo->BodyPairs[2*PairIdx];
 		FName OtherBone = CachedNotifyInfo->BodyPairs[2*PairIdx + 1];
 
-		int32 PrimaryIdx = CheckBonesMap[PrimaryBone];
-		int32 OtherIdx = CheckBonesMap[OtherBone];
+		int32 PrimaryIdx = CheckSourceBonesMap[PrimaryBone];
+		int32 OtherIdx = CheckSourceBonesMap[OtherBone];
 
 		// Secondary (contact) Bone always affected
-		CacheBodyEffectVertIdx[OtherIdx].Add(2*PairIdx+1);
+		CacheSourceBodyEffectVertIdx[OtherIdx].Add(2*PairIdx+1);
 		if (!bPrimaryFixed)
 		{
 			// Primary bone can be affect if e.g. hands clapping
-			CacheBodyEffectVertIdx[PrimaryIdx].Add(2*PairIdx);
+			CacheSourceBodyEffectVertIdx[PrimaryIdx].Add(2*PairIdx);
 		}
 	}
 }
@@ -819,7 +912,7 @@ void FIKRetargetRelativeIKOp::DebugDraw(
 		const float CoordThickness = 1.0f;
 		
 		// Show Source Body coords
-		for (const FTransform& BodyTfm : DebugDrawSourceBodyTfms)
+		for (const FTransform& BodyTfm : DebugDrawInfo.SourceBodyTfms)
 		{
 			FTransform BodyFrame = BodyTfm * SourceNoScaleTfm;
 			DrawCoordinateSystem(
@@ -832,7 +925,7 @@ void FIKRetargetRelativeIKOp::DebugDraw(
 		}
 
 		// Show Target Body coords
-		for (const FTransform& BodyTfm : DebugDrawTargetBodyTfms)
+		for (const FTransform& BodyTfm : DebugDrawInfo.TargetBodyTfms)
 		{
 			FTransform BodyFrame = BodyTfm * InComponentTransform;
 			DrawCoordinateSystem(
@@ -848,10 +941,10 @@ void FIKRetargetRelativeIKOp::DebugDraw(
 	if (Settings.bDebugDrawRetargetVertAverages)
 	{
 		// Show retarget contribution average with offsets
-		for (const auto& t : DebugDrawTargetContrib)
+		for (const FDebugRelativeTargetPairSpace& t : DebugDrawInfo.PairContrib)
 		{
-			FTransform V1Tfm = FTransform(t.Get<0>()) * InComponentTransform;
-			FTransform V2Tfm = FTransform(t.Get<1>()) * InComponentTransform;
+			FTransform V1Tfm = FTransform(t.PairRangeStart) * InComponentTransform;
+			FTransform V2Tfm = FTransform(t.PairRangeEnd) * InComponentTransform;
 			IKRigDebugRendering::DrawWireCube(InPDI, V1Tfm, FLinearColor::Red,static_cast<float>(0.5*InComponentScale),static_cast<float>(0.5*InComponentScale));
 			IKRigDebugRendering::DrawWireCube(InPDI, V2Tfm, FLinearColor::Red,static_cast<float>(0.5*InComponentScale),static_cast<float>(0.5*InComponentScale));
 
@@ -861,13 +954,15 @@ void FIKRetargetRelativeIKOp::DebugDraw(
 			// DrawDashedLine(InPDI, VBone, V2, FLinearColor::Yellow, 1, SDPG_Foreground);
 			DrawDashedLine(InPDI, V1, V2, FLinearColor::Black, 1, SDPG_Foreground);
 
-			FTransform VAvgTfm = FTransform(t.Get<2>()) * InComponentTransform;
+			FTransform VAvgTfm = FTransform(t.PairTarget) * InComponentTransform;
 			IKRigDebugRendering::DrawWireCube(InPDI, VAvgTfm, FLinearColor::Gray,static_cast<float>(0.5*InComponentScale),static_cast<float>(0.5*InComponentScale));
 
-			FTransform VMinTfm = FTransform(t.Get<3>()) * InComponentTransform;
+			FVector FeasibleMinV =  FMath::Lerp(t.PairRangeStart, t.PairRangeEnd, t.FeasibleRange.Min);
+			FTransform VMinTfm = FTransform(FeasibleMinV) * InComponentTransform;
 			IKRigDebugRendering::DrawWireCube(InPDI, VMinTfm, FLinearColor::Yellow,static_cast<float>(0.5*InComponentScale),static_cast<float>(0.5*InComponentScale));
 
-			FTransform VMaxTfm = FTransform(t.Get<4>()) * InComponentTransform;
+			FVector FeasibleMaxV =  FMath::Lerp(t.PairRangeStart, t.PairRangeEnd, t.FeasibleRange.Max);
+			FTransform VMaxTfm = FTransform(FeasibleMaxV) * InComponentTransform;
 			IKRigDebugRendering::DrawWireCube(InPDI, VMaxTfm, FLinearColor::Green,static_cast<float>(0.5*InComponentScale),static_cast<float>(0.5*InComponentScale));
 			// DrawDashedLine(InPDI, VOffset, VOffset + t.Get<3>(), FLinearColor::Black, 1, SDPG_Foreground);
 		}
@@ -876,11 +971,11 @@ void FIKRetargetRelativeIKOp::DebugDraw(
 	if (Settings.bDebugDrawGoalContributions)
 	{
 		// Show IK Goals (and contribs)
-		for (const auto& t : DebugDrawGoals)
+		for (const FDebugDrawTargetGoal& t : DebugDrawInfo.TargetGoals)
 		{
-			FTransform GoalTfm = FTransform(t.Get<0>()) * InComponentTransform;
+			FTransform GoalTfm = FTransform(t.Goal) * InComponentTransform;
 			FVector GoalPoint = GoalTfm.GetTranslation();
-			for (const FVector& v: t.Get<1>())
+			for (const FVector& v: t.PairTargets)
 			{
 				FTransform ContribTfm = FTransform(v) * InComponentTransform;
 				FVector ContribPt = ContribTfm.GetTranslation();
@@ -894,10 +989,10 @@ void FIKRetargetRelativeIKOp::DebugDraw(
 	if (Settings.bDebugDrawBodyPairs)
 	{
 		// Draw Source body pairs
-		for (int Index = 0; Index < DebugDrawSourcePairVerts.Num(); ++Index)
+		for (int Index = 0; Index < DebugDrawInfo.SourcePairVerts.Num(); ++Index)
 		{
-			FTransform PrimTfm = FTransform(DebugDrawSourcePairVerts[Index].Get<0>()) * SourceScaleCheckedTfm;
-			FTransform SecTfm = FTransform(DebugDrawSourcePairVerts[Index].Get<1>()) * SourceScaleCheckedTfm;
+			FTransform PrimTfm = FTransform(DebugDrawInfo.SourcePairVerts[Index].Get<0>()) * SourceScaleCheckedTfm;
+			FTransform SecTfm = FTransform(DebugDrawInfo.SourcePairVerts[Index].Get<1>()) * SourceScaleCheckedTfm;
 
 			FVector v1 = PrimTfm.GetTranslation();
 			FVector v2 = SecTfm.GetTranslation();
@@ -908,13 +1003,13 @@ void FIKRetargetRelativeIKOp::DebugDraw(
 		}
 
 		// Draw target body pairs (relative to original retarget body)
-		for (int Index = 0; Index < DebugDrawTargetPairVerts.Num(); Index += 2)
+		for (int Index = 0; Index < DebugDrawInfo.TargetPairVerts.Num(); Index += 2)
 		{
-			FTransform PrimTfmA = FTransform(DebugDrawTargetPairVerts[Index].Get<0>()) * InComponentTransform;
-			FTransform SecTfmA = FTransform(DebugDrawTargetPairVerts[Index].Get<1>()) * InComponentTransform;
+			FTransform PrimTfmA = FTransform(DebugDrawInfo.TargetPairVerts[Index].Get<0>()) * InComponentTransform;
+			FTransform SecTfmA = FTransform(DebugDrawInfo.TargetPairVerts[Index].Get<1>()) * InComponentTransform;
 
-			FTransform PrimTfmB = FTransform(DebugDrawTargetPairVerts[Index+1].Get<0>()) * InComponentTransform;
-			FTransform SecTfmB = FTransform(DebugDrawTargetPairVerts[Index+1].Get<1>()) * InComponentTransform;
+			FTransform PrimTfmB = FTransform(DebugDrawInfo.TargetPairVerts[Index+1].Get<0>()) * InComponentTransform;
+			FTransform SecTfmB = FTransform(DebugDrawInfo.TargetPairVerts[Index+1].Get<1>()) * InComponentTransform;
 
 			FVector v1 = PrimTfmA.GetTranslation();
 			FVector v2 = SecTfmA.GetTranslation();
@@ -934,32 +1029,37 @@ void FIKRetargetRelativeIKOp::DebugDraw(
 		}
 	}
 
-	if (Settings.bDebugDrawIntersectTest)
-	{
-		for (int32 BodyIdx = 0; BodyIdx < DebugDrawIntersectTfms.Num(); ++BodyIdx)
-		{
-			if (!Settings.IntersectBodies.IsValidIndex(BodyIdx))
-			{
-				continue;
-			}
-			
-			FName BodyName = Settings.IntersectBodies[BodyIdx];
-			FKShapeElem* ShapeElem = FindBodyShape(TargetPhysicsAsset, BodyName);
-			if (!ShapeElem || ShapeElem->GetShapeType() != EAggCollisionShape::Sphyl)
-			{
-				continue;
-			}
-	
-			const FKSphylElem* CapsuleElem = static_cast<FKSphylElem*>(ShapeElem);
-			FTransform Tfm = CapsuleElem->GetTransform() * DebugDrawIntersectTfms[BodyIdx] * InComponentTransform;
-			const FVector Translation = Tfm.GetTranslation();
-			const FVector UnitXAxis   = Tfm.GetUnitAxis( EAxis::X );
-			const FVector UnitYAxis   = Tfm.GetUnitAxis( EAxis::Y );
-			const FVector UnitZAxis   = Tfm.GetUnitAxis( EAxis::Z );
-	
-			// DrawWireCapsule(InPDI, Translation, UnitXAxis, UnitYAxis, UnitZAxis, FLinearColor::Yellow, CapsuleElem->GetScaledRadius()*InComponentScale, CapsuleElem->GetScaledHalfLength()*InComponentScale, 25, SDPG_Foreground);
-		}
-	}
+	// if (Settings.bDebugDrawIntersectTest)
+	// {
+	// 	for (int32 BodyIdx = 0; BodyIdx < DebugDrawInfo.TargetIntersectTfms.Num(); ++BodyIdx)
+	// 	{
+	// 		if (!Settings.IntersectBodies.IsValidIndex(BodyIdx))
+	// 		{
+	// 			continue;
+	// 		}
+	// 		
+	// 		FName BodyName = Settings.IntersectBodies[BodyIdx];
+	// 		FKShapeElem* ShapeElem = FindBodyShape(TargetPhysicsAsset, BodyName);
+	// 		if (!ShapeElem || ShapeElem->GetShapeType() != EAggCollisionShape::Sphyl)
+	// 		{
+	// 			continue;
+	// 		}
+	//
+	// 		const FKSphylElem* CapsuleElem = static_cast<FKSphylElem*>(ShapeElem);
+	// 		FTransform Tfm = CapsuleElem->GetTransform() * DebugDrawInfo.TargetIntersectTfms[BodyIdx] * InComponentTransform;
+	// 		const FVector Translation = Tfm.GetTranslation();
+	// 		const FVector UnitXAxis   = Tfm.GetUnitAxis( EAxis::X );
+	// 		const FVector UnitYAxis   = Tfm.GetUnitAxis( EAxis::Y );
+	// 		const FVector UnitZAxis   = Tfm.GetUnitAxis( EAxis::Z );
+	//
+	// 		DrawWireCapsule(InPDI, Translation,
+	// 			UnitXAxis, UnitYAxis, UnitZAxis,
+	// 			FLinearColor::Yellow,
+	// 			CapsuleElem->Radius * InComponentScale,
+	// 			(0.5*CapsuleElem->Length + CapsuleElem->Radius) * InComponentScale,
+	// 			25, SDPG_Foreground, 0, 1.0);
+	// 	}
+	// }
 }
 
 #endif //WITH_EDITOR
diff --git a/Engine/Plugins/Experimental/Animation/RelativeIKOp/Source/RelativeIKOp/Public/RelativeIKOp.h b/Engine/Plugins/Experimental/Animation/RelativeIKOp/Source/RelativeIKOp/Public/RelativeIKOp.h
--- a/Engine/Plugins/Experimental/Animation/RelativeIKOp/Source/RelativeIKOp/Public/RelativeIKOp.h
+++ b/Engine/Plugins/Experimental/Animation/RelativeIKOp/Source/RelativeIKOp/Public/RelativeIKOp.h
@@ -17,6 +17,45 @@ struct FIKRetargetRunIKRigOp;
 class UPhysicsAsset;
 class URelativeBodyBakeAnimNotify;
 
+// Debug info for individual body pair contributions
+struct FDebugRelativeTargetPairSpace
+{
+	FVector PairTarget;
+	
+	FVector PairRangeStart;
+	FVector PairRangeEnd;
+	
+	FDoubleInterval FeasibleRange;
+};
+
+// Debug draw info for goal (with individual pair goal targets)
+struct FDebugDrawTargetGoal
+{
+	FVector Goal;
+	TArray<FVector> PairTargets;
+};
+
+// Used to hold all debug draw info for convenience
+struct FDebugRelativeIKDrawInfo
+{
+	// Source/target vert pair info
+	TArray<TTuple<FVector,FVector>> SourcePairVerts;
+	TArray<TTuple<FVector,FVector>> TargetPairVerts;
+
+	// Body coordinates
+	TArray<FTransform> SourceBodyTfms;
+	TArray<FTransform> TargetBodyTfms;
+
+	// Intersect info
+	TArray<FTransform> TargetIntersectTfms;
+	TArray<TTuple<FVector,double>> TestSpheres;
+
+	// Debug info for pair contributions
+	TArray<FDebugRelativeTargetPairSpace> PairContrib;
+	// Final goal points with contributions
+	TArray<FDebugDrawTargetGoal> TargetGoals;
+};
+
 USTRUCT(BlueprintType, meta = (DisplayName = "Relative IK Settings"))
 struct FIKRetargetRelativeIKOpSettings : public FIKRetargetOpSettingsBase
 {
@@ -34,6 +73,10 @@ struct FIKRetargetRelativeIKOpSettings : public FIKRetargetOpSettingsBase
 	UPROPERTY(EditAnywhere, BlueprintReadWrite, Category = "Physics", meta=(ReinitializeOnEdit))
 	TObjectPtr<UPhysicsAsset> TargetPhysicsAssetOverride;
 
+	// Physics body source -> target name mapping 
+	UPROPERTY(EditAnywhere, BlueprintReadWrite, Category = "Source Target Mapping")
+	TMap<FName,FName> BodyMapping;
+
 	// Physics body names that should maintain relative spatial relationships with ContactBodies 
 	UPROPERTY(EditAnywhere, BlueprintReadWrite, Category = "Body Pairs")
 	TArray<FName> DomainBodies;
@@ -92,14 +135,14 @@ struct FIKRetargetRelativeIKOpSettings : public FIKRetargetOpSettingsBase
 	UPROPERTY(EditAnywhere, BlueprintReadWrite, Category = "Debug")
 	bool bDebugDrawBodyCoords = false;
 
-	UPROPERTY(EditAnywhere, BlueprintReadWrite, Category = "Debug")
-	bool bDebugDrawDomainBodies = false;
-	
-	UPROPERTY(EditAnywhere, BlueprintReadWrite, Category = "Debug")
-	bool bDebugDrawContactBodies = false;
-
-	UPROPERTY(EditAnywhere, BlueprintReadWrite, Category = "Debug")
-	bool bDebugDrawIntersectTest = false;
+	// UPROPERTY(EditAnywhere, BlueprintReadWrite, Category = "Debug")
+	// bool bDebugDrawDomainBodies = false;
+	//
+	// UPROPERTY(EditAnywhere, BlueprintReadWrite, Category = "Debug")
+	// bool bDebugDrawContactBodies = false;
+	//
+	// UPROPERTY(EditAnywhere, BlueprintReadWrite, Category = "Debug")
+	// bool bDebugDrawIntersectTest = false;
 
 	UPROPERTY(EditAnywhere, BlueprintReadWrite, Category = "Test Features")
 	bool bTestDistContactAlpha = true;
@@ -148,11 +191,17 @@ struct FIKRetargetRelativeIKOp : public FIKRetargetOpBase
 
 private:
 	void UpdateCacheBoneChains(const FIKRetargetProcessor& InProcessor, const FIKRetargetRunIKRigOp* ParentRigOp);
-	void UpdateBoneMaps(FName BoneName);
+	void UpdateAllBoneMaps(FName SourceBoneName);
+	void UpdateTargetBoneMap(FName TargetBoneName);
 	void UpdateCacheSkelInfo(const FRetargetSkeleton& InSourceSkeleton, const FTargetSkeleton& InTargetSkeleton);
 	void PreUpdateRelativeIKNotifyInfo();
 	void UpdateCacheNotifyInfo(URelativeBodyBakeAnimNotify* NotifyInfo, float AnimStart, float AnimEnd);
 
+	FName ApplyBodyMap(FName BodyName);
+	void UpdateCacheChainInfo(const TArray<FTransform>& TargetPose, FName BoneName);
+	void ComputeFeasibilityRange(FDoubleInterval& OutFeasibleRange, FName BoneName, FVector BoneChainOffset, FVector TargetRelativeToBoneV, FVector TargetRelativeToOpposeBoneV);
+	// void FixupSphereIntersections();
+
 	static double CalcLineSphereIntersect(const FVector& Center, double Radius, const FVector& StartPoint, const FVector& EndPoint);
 	static double CalcIntersectionDelta(const FVector&  TargetPoint, double TargetRadius, const FTransform& ShapeTfm, FKShapeElem* ShapeElem, FVector& OutDeltaDir);
 	
@@ -162,11 +211,11 @@ private:
 	static FTransform CalcReferenceShapeScale3D(FKShapeElem* ShapeElem);
 	static FTransform GetBodyTransform(const UPhysicsAsset* PhysAsset, FName BoneName);
 	static FKShapeElem* FindBodyShape(const UPhysicsAsset* PhysAsset, FName BoneName);
+	
 	static double GetDistanceWeight(double Distance, double DistThreshold, double ScalarFade = 200.0);
 
 #if WITH_EDITOR
 public:
-	// TODO: Update DebugDraw API to provide source transform as well
 	UE_API virtual void DebugDraw(
 		FPrimitiveDrawInterface* InPDI,
 		const FTransform& InSourceTransform,
@@ -174,19 +223,10 @@ public:
 		const double InComponentScale,
 		const FIKRetargetDebugDrawState& InEditorState) const override;
 
-	TArray<TTuple<FVector,FVector>> DebugDrawSourcePairVerts;
-	TArray<TTuple<FVector,FVector>> DebugDrawTargetPairVerts;
-
-	TArray<FTransform> DebugDrawSourceBodyTfms;
-	TArray<FTransform> DebugDrawTargetBodyTfms;
-	TArray<FTransform> DebugDrawIntersectTfms;
-	TArray<TTuple<FVector,double>> DebugDrawTestSpheres;
-
-	TArray<TTuple<FVector,FVector,FVector,FVector,FVector>> DebugDrawTargetContrib;
-	TArray<TTuple<FVector,TArray<FVector>>> DebugDrawGoals;
+	// Collection of all debug structures
+	FDebugRelativeIKDrawInfo DebugDrawInfo;
 	
 	static UE_API FCriticalSection DebugDataMutex;
-	// Add debug draw data
 #endif
 private:
 	
@@ -206,8 +246,8 @@ private:
 	TArray<FName> TargetBoneNames;
 
 	// Array of body-vert indices affecting each joint goal (for weighted averaging)
-	TArray<TArray<int32>> CacheBodyEffectVertIdx;
-	TArray<FName> CacheEffectBones;
+	TArray<TArray<int32>> CacheSourceBodyEffectVertIdx;
+	TArray<FName> CacheSourceEffectBones;
 	TArray<double> FrameBoneWeights;
 	TArray<FVector> FrameBoneTargets;
 
@@ -216,6 +256,9 @@ private:
 	TMap<FName, const FResolvedBoneChain*> CacheBoneChains;
 	TMap<FName,int32> CacheSourceSkelIndices;
 	TMap<FName,int32> CacheTargetSkelIndices;
+	// Additional per-frame ik cache maps for length lookup
+	TMap<FName,double> CacheChainLengthMap;
+	TMap<FName,FVector> CacheChainStartMap;
 	
 	FAnimMontageInstance* MontageInstance = nullptr;
 };
