RelativeIKOp: Only require mapping for different bodies (default identity mapping).

diff --git a/Engine/Plugins/Experimental/Animation/RelativeIKOp/Source/RelativeIKOp/Private/RelativeIKOp.cpp b/Engine/Plugins/Experimental/Animation/RelativeIKOp/Source/RelativeIKOp/Private/RelativeIKOp.cpp
--- a/Engine/Plugins/Experimental/Animation/RelativeIKOp/Source/RelativeIKOp/Private/RelativeIKOp.cpp
+++ b/Engine/Plugins/Experimental/Animation/RelativeIKOp/Source/RelativeIKOp/Private/RelativeIKOp.cpp
@@ -415,18 +415,11 @@ void FIKRetargetRelativeIKOp::Run(
 
 FName FIKRetargetRelativeIKOp::ApplyBodyMap(FName BodyName)
 {
-	if (Settings.BodyMapping.IsEmpty())
+	if (Settings.BodyMapping.IsEmpty() || !Settings.BodyMapping.Contains(BodyName))
 	{
 		return BodyName;
 	}
 
-	// TODO: Debug-only logging of issues
-	if (!Settings.BodyMapping.Contains(BodyName))
-	{
-		// UE_LOG(LogAnimation, Warning, TEXT("Body name in baked data but not in mapping: %s"), *BodyName.ToString());
-		return NAME_None;
-	}
-
 	return Settings.BodyMapping[BodyName];
 }
 
