RelativeIKOp: Testing body-rotation invariance with emulated space baking (WIP).
NOTE: this version remains dependent on the body orientations in the source physics asset, since data isn't being rebaked

diff --git a/Engine/Plugins/Experimental/Animation/RelativeIKOp/Source/RelativeIKOp/Private/RelativeIKOp.cpp b/Engine/Plugins/Experimental/Animation/RelativeIKOp/Source/RelativeIKOp/Private/RelativeIKOp.cpp
--- a/Engine/Plugins/Experimental/Animation/RelativeIKOp/Source/RelativeIKOp/Private/RelativeIKOp.cpp
+++ b/Engine/Plugins/Experimental/Animation/RelativeIKOp/Source/RelativeIKOp/Private/RelativeIKOp.cpp
@@ -162,14 +162,22 @@ void FIKRetargetRelativeIKOp::Run(
 		int32 SourcePrimaryPoseIdx = CacheSourceSkelIndices[SourcePrimaryBone];
 		int32 SourceOtherPoseIdx = CacheSourceSkelIndices[SourceOtherBone];
 
+		// ---- TESTING ONLY - Flip these once notify baking is rotation invariant
+		FTransform SourcePrimaryBodyRot = GetBodyRotation(SourcePhysicsAsset, SourcePrimaryBone);
+		FTransform SourceOtherBodyRot = GetBodyRotation(SourcePhysicsAsset, SourceOtherBone);
+		FTransform TargetPrimaryBodyRot = (Settings.bTestBodyRotInvariance) ? FTransform::Identity : SourcePrimaryBodyRot.Inverse() * GetBodyRotation(TargetPhysicsAsset, TargetPrimaryBone);
+		FTransform TargetOtherBodyRot = (Settings.bTestBodyRotInvariance) ? FTransform::Identity : SourceOtherBodyRot.Inverse() * GetBodyRotation(TargetPhysicsAsset, TargetOtherBone);
+		
 		// Compute source transforms and source-space vertex distance
 		const FVector3f& LocalPrimaryVert = CachedNotifyInfo->BodyPairsLocalReference[FramePairIdx + 2*PairIdx];
-		FTransform SourcePrimaryBodyTfm = GetBodyTransform(SourcePhysicsAsset, SourcePrimaryBone) * SourceScaleTfm;
+		// TODO: Swap and clean these up if we bake in rot invaraint space
+		// FTransform SourcePrimaryBodyTfm = GetBodyTransformInvariant(SourcePhysicsAsset, SourcePrimaryBone) * SourceScaleTfm;
+		FTransform SourcePrimaryBodyTfm = SourcePrimaryBodyRot * GetBodyTransformInvariant(SourcePhysicsAsset, SourcePrimaryBone) * SourceScaleTfm;
 		FTransform SourcePrimaryCompTfm = SourcePrimaryBodyTfm * InSourceGlobalPose[SourcePrimaryPoseIdx] * InvScaleTfm;
 		FVector SourcePrimaryCompV = SourcePrimaryCompTfm.TransformPosition(FVector(LocalPrimaryVert));
 
 		const FVector3f& LocalOtherVert = CachedNotifyInfo->BodyPairsLocalReference[FramePairIdx + 2*PairIdx + 1];
-		FTransform SourceOtherBodyTfm = GetBodyTransform(SourcePhysicsAsset, SourceOtherBone) * SourceScaleTfm;
+		FTransform SourceOtherBodyTfm = SourceOtherBodyRot * GetBodyTransformInvariant(SourcePhysicsAsset, SourceOtherBone) * SourceScaleTfm;
 		FTransform SourceOtherCompTfm = SourceOtherBodyTfm * InSourceGlobalPose[SourceOtherPoseIdx] * InvScaleTfm;
 		FVector SourceOtherCompV = SourceOtherCompTfm.TransformPosition(FVector(LocalOtherVert));
 
@@ -184,8 +192,10 @@ void FIKRetargetRelativeIKOp::Run(
 		// Get average of retarget verts
 		int32 TargetPrimaryPoseIdx = CacheTargetSkelIndices[TargetPrimaryBone];
 		int32 TargetOtherPoseIdx = CacheTargetSkelIndices[TargetOtherBone];
-		FTransform TargetPrimaryBodyTfm = GetBodyTransform(TargetPhysicsAsset, TargetPrimaryBone);
-		FTransform TargetOtherBodyTfm = GetBodyTransform(TargetPhysicsAsset, TargetOtherBone);
+		// TODO: Swap and clean these up if we bake in rot invaraint space
+		// FTransform TargetPrimaryBodyTfm = GetBodyTransformInvariant(TargetPhysicsAsset, TargetPrimaryBone);
+		FTransform TargetPrimaryBodyTfm = SourcePrimaryBodyRot * TargetPrimaryBodyRot * GetBodyTransformInvariant(TargetPhysicsAsset, TargetPrimaryBone);
+		FTransform TargetOtherBodyTfm = SourceOtherBodyRot * TargetOtherBodyRot * GetBodyTransformInvariant(TargetPhysicsAsset, TargetOtherBone);
 
 		FTransform TargetPrimaryCompTfm = TargetPrimaryBodyTfm * OutTargetGlobalPose[TargetPrimaryPoseIdx];
 		FTransform TargetOtherCompTfm = TargetOtherBodyTfm * OutTargetGlobalPose[TargetOtherPoseIdx];
@@ -388,12 +398,14 @@ void FIKRetargetRelativeIKOp::Run(
 			}
 			
 			int32 SourcePoseIdx = CacheSourceSkelIndices[SourceBoneName];
-			FTransform SourceBodyTfm = GetBodyTransform(SourcePhysicsAsset, SourceBoneName) * SourceScaleTfm;
+			FTransform SourceBodyRot = (Settings.bTestBodyRotInvariance) ? FTransform::Identity : GetBodyRotation(SourcePhysicsAsset, SourceBoneName);
+			FTransform SourceBodyTfm = SourceBodyRot * GetBodyTransformInvariant(SourcePhysicsAsset, SourceBoneName) * SourceScaleTfm;
 			FTransform SourceBodyCompTfm = SourceBodyTfm * InSourceGlobalPose[SourcePoseIdx];
 			DebugDrawInfo.SourceBodyTfms.Add(SourceBodyCompTfm);
 			
 			int32 TargetPoseIdx = CacheTargetSkelIndices[TargetBoneName];
-			FTransform TargetBodyTfm = GetBodyTransform(TargetPhysicsAsset, TargetBoneName);
+			FTransform TargetBodyRot = (Settings.bTestBodyRotInvariance) ? FTransform::Identity :GetBodyRotation(TargetPhysicsAsset, TargetBoneName);
+			FTransform TargetBodyTfm = TargetBodyRot * GetBodyTransformInvariant(TargetPhysicsAsset, TargetBoneName);
 			FTransform TargetBodyCompTfm = TargetBodyTfm * OutTargetGlobalPose[TargetPoseIdx];
 			DebugDrawInfo.TargetBodyTfms.Add(TargetBodyCompTfm);
 		}
@@ -625,7 +637,19 @@ FTransform FIKRetargetRelativeIKOp::CalcReferenceShapeScale3D(FKShapeElem* Shape
 	return ElementScale3D;
 }
 
-FTransform FIKRetargetRelativeIKOp::GetBodyTransform(const UPhysicsAsset* PhysAsset, FName BoneName)
+FTransform FIKRetargetRelativeIKOp::GetBodyRotation(const UPhysicsAsset* PhysAsset, FName BoneName)
+{
+	FKShapeElem* ShapeElem = FindBodyShape(PhysAsset, BoneName);
+	if (!ShapeElem)
+	{
+		return FTransform::Identity;
+	}
+	
+	FTransform BodyTransform = ShapeElem->GetTransform();
+	return FTransform(BodyTransform.GetRotation());
+}
+
+FTransform FIKRetargetRelativeIKOp::GetBodyTransformInvariant(const UPhysicsAsset* PhysAsset, FName BoneName)
 {
 	FKShapeElem* ShapeElem = FindBodyShape(PhysAsset, BoneName);
 	if (!ShapeElem)
@@ -634,8 +658,9 @@ FTransform FIKRetargetRelativeIKOp::GetBodyTransform(const UPhysicsAsset* PhysAs
 	}
 	
 	FTransform BodyTransform = ShapeElem->GetTransform();
+	FTransform BodyRotInv(BodyTransform.GetRotation().Inverse());
 	FTransform ElementScale3D = CalcReferenceShapeScale3D(ShapeElem);
-	return ElementScale3D * BodyTransform;
+	return BodyRotInv * ElementScale3D * BodyTransform;
 }
 
 FKShapeElem* FIKRetargetRelativeIKOp::FindBodyShape(const UPhysicsAsset* PhysAsset, FName BoneName)
diff --git a/Engine/Plugins/Experimental/Animation/RelativeIKOp/Source/RelativeIKOp/Public/RelativeIKOp.h b/Engine/Plugins/Experimental/Animation/RelativeIKOp/Source/RelativeIKOp/Public/RelativeIKOp.h
--- a/Engine/Plugins/Experimental/Animation/RelativeIKOp/Source/RelativeIKOp/Public/RelativeIKOp.h
+++ b/Engine/Plugins/Experimental/Animation/RelativeIKOp/Source/RelativeIKOp/Public/RelativeIKOp.h
@@ -152,6 +152,9 @@ struct FIKRetargetRelativeIKOpSettings : public FIKRetargetOpSettingsBase
 
 	UPROPERTY(EditAnywhere, BlueprintReadWrite, Category = "Test Features")
 	bool bTestIntersectionFixup = true;
+
+	UPROPERTY(EditAnywhere, BlueprintReadWrite, Category = "Test Features")
+	bool bTestBodyRotInvariance = true;
 };
 
 USTRUCT(BlueprintType, meta = (DisplayName = "Relative IK Goal Op"))
@@ -209,7 +212,8 @@ private:
 	static double CalcShapeAvgRadius(FKShapeElem* ShapeElem);
 
 	static FTransform CalcReferenceShapeScale3D(FKShapeElem* ShapeElem);
-	static FTransform GetBodyTransform(const UPhysicsAsset* PhysAsset, FName BoneName);
+	static FTransform GetBodyRotation(const UPhysicsAsset* PhysAsset, FName BoneName);
+	static FTransform GetBodyTransformInvariant(const UPhysicsAsset* PhysAsset, FName BoneName);
 	static FKShapeElem* FindBodyShape(const UPhysicsAsset* PhysAsset, FName BoneName);
 	
 	static double GetDistanceWeight(double Distance, double DistThreshold, double ScalarFade = 200.0);
