RelativeIKOp: Fix debug draw sticking around if clicking to anim witout notify info.

diff --git a/Engine/Plugins/Experimental/Animation/RelativeIKOp/Source/RelativeIKOp/Private/RelativeIKOp.cpp b/Engine/Plugins/Experimental/Animation/RelativeIKOp/Source/RelativeIKOp/Private/RelativeIKOp.cpp
--- a/Engine/Plugins/Experimental/Animation/RelativeIKOp/Source/RelativeIKOp/Private/RelativeIKOp.cpp
+++ b/Engine/Plugins/Experimental/Animation/RelativeIKOp/Source/RelativeIKOp/Private/RelativeIKOp.cpp
@@ -87,7 +87,10 @@ void FIKRetargetRelativeIKOp::Run(
 	const TArray<FTransform>& InSourceGlobalPose,
 	TArray<FTransform>& OutTargetGlobalPose)
 {
-	// TODO: Reset/clear debug info structs here to
+#if WITH_EDITOR
+	ResetDebugInfo();
+#endif //WITH_EDITOR
+	
 	if (InProcessor.IsIKForcedOff())
 	{
 		return;
@@ -1080,6 +1083,21 @@ void FIKRetargetRelativeIKOp::DebugDraw(
 	// }
 }
 
+void FIKRetargetRelativeIKOp::ResetDebugInfo()
+{
+	DebugDrawInfo.SourcePairVerts.Reset();
+	DebugDrawInfo.TargetPairVerts.Reset();
+	
+	DebugDrawInfo.SourceBodyTfms.Reset();
+	DebugDrawInfo.TargetBodyTfms.Reset();
+	
+	DebugDrawInfo.TargetIntersectTfms.Reset();
+	DebugDrawInfo.TestSpheres.Reset();
+	
+	DebugDrawInfo.PairContrib.Reset();
+	DebugDrawInfo.TargetGoals.Reset();
+}
+
 #endif //WITH_EDITOR
 
 FIKRetargetRelativeIKOpSettings UIKRetargetRelativeIKController::GetSettings()
diff --git a/Engine/Plugins/Experimental/Animation/RelativeIKOp/Source/RelativeIKOp/Public/RelativeIKOp.h b/Engine/Plugins/Experimental/Animation/RelativeIKOp/Source/RelativeIKOp/Public/RelativeIKOp.h
--- a/Engine/Plugins/Experimental/Animation/RelativeIKOp/Source/RelativeIKOp/Public/RelativeIKOp.h
+++ b/Engine/Plugins/Experimental/Animation/RelativeIKOp/Source/RelativeIKOp/Public/RelativeIKOp.h
@@ -227,6 +227,9 @@ public:
 		const double InComponentScale,
 		const FIKRetargetDebugDrawState& InEditorState) const override;
 
+private:
+	void ResetDebugInfo();
+
 	// Collection of all debug structures
 	FDebugRelativeIKDrawInfo DebugDrawInfo;
 	
