RelativeIKOp: Cleanup debug code and added retarget body pair display w/ advanced list and ability to show phys assets for baked data (WIP).

diff --git a/Engine/Plugins/Experimental/Animation/RelativeIKOp/Source/RelativeIKOp/Private/RelativeIKOp.cpp b/Engine/Plugins/Experimental/Animation/RelativeIKOp/Source/RelativeIKOp/Private/RelativeIKOp.cpp
--- a/Engine/Plugins/Experimental/Animation/RelativeIKOp/Source/RelativeIKOp/Private/RelativeIKOp.cpp
+++ b/Engine/Plugins/Experimental/Animation/RelativeIKOp/Source/RelativeIKOp/Private/RelativeIKOp.cpp
@@ -250,11 +250,14 @@ void FIKRetargetRelativeIKOp::Run(
 		// Debug draw verts
 		if (Settings.bEnableDebugDraw)
 		{
-			LocalDebugInfo.SourcePairVerts.Add(TTuple<FVector,FVector>(SourcePrimaryCompV, SourceOtherCompV));
-			if (Settings.DrawTargetPairBones.Contains(TargetPrimaryBone))
+			LocalDebugInfo.SourcePairVerts.Add({SourcePrimaryCompV, SourceOtherCompV, false, FVector::ZeroVector, FVector::ZeroVector});
+			if (Settings.DebugFullRetargetPairBones.Contains(TargetPrimaryBone))
 			{
-				LocalDebugInfo.TargetPairVerts.Add(TTuple<FVector,FVector>(TargetPrimaryCompV_P, TargetOtherCompV_O));
-				LocalDebugInfo.TargetPairVerts.Add(TTuple<FVector,FVector>(TargetPrimaryCompV_O, TargetOtherCompV_P));
+				LocalDebugInfo.TargetPairVerts.Add({TargetPrimaryCompV_P, TargetOtherCompV_O, true,TargetOtherCompV_P, TargetPrimaryCompV_O,});
+			}
+			else
+			{
+				LocalDebugInfo.TargetPairVerts.Add({TargetPrimaryCompV_P, TargetOtherCompV_O, false, FVector::ZeroVector, FVector::ZeroVector});
 			}
 			
 			LocalTargetPairSpaces[2*PairIdx] = FDebugRelativeTargetPairSpace{WeightedTargetPrimary, TargetRelativeToOtherV+RelativePrimaryOffset, TargetRelativeToPrimaryV+RelativePrimaryOffset, FeasibleRange};
@@ -308,7 +311,7 @@ void FIKRetargetRelativeIKOp::Run(
 			if (Settings.bEnableDebugDraw)
 			{
 				LocalDebugInfo.TargetGoals.Last().PairTargets.Add(FrameBoneTargets[Index]);
-				LocalDebugInfo.PairContrib.Emplace(LocalTargetPairSpaces[Index]);
+				LocalDebugInfo.PairRetargetInfo.Emplace(LocalTargetPairSpaces[Index]);
 			}
 #endif //WITH_EDITOR
 		}
@@ -385,45 +388,51 @@ void FIKRetargetRelativeIKOp::Run(
 		FScopeLock ScopeLock(&DebugDataMutex);
 
 		DebugDrawInfo.TargetGoals = LocalDebugInfo.TargetGoals;
-		DebugDrawInfo.PairContrib = LocalDebugInfo.PairContrib;
+		DebugDrawInfo.PairRetargetInfo = LocalDebugInfo.PairRetargetInfo;
 		DebugDrawInfo.SourcePairVerts = LocalDebugInfo.SourcePairVerts;
 		DebugDrawInfo.TargetPairVerts = LocalDebugInfo.TargetPairVerts;
 		DebugDrawInfo.TestSpheres = LocalDebugInfo.TestSpheres;
 
-		DebugDrawInfo.SourceBodyTfms.Reset(CacheSourceEffectBones.Num());
-		DebugDrawInfo.TargetBodyTfms.Reset(CacheSourceEffectBones.Num());
+		// Cache debug info for source bodies
+		DebugDrawInfo.SourceBodyInfo.Reset(CacheSourceEffectBones.Num());
 		for (FName SourceBoneName : CacheSourceEffectBones)
 		{
-			FName TargetBoneName = ApplyBodyMap(SourceBoneName);
-			if (!CacheSourceSkelIndices.Contains(SourceBoneName) || !CacheTargetSkelIndices.Contains(TargetBoneName))
+			if (!CacheSourceSkelIndices.Contains(SourceBoneName))
 			{
 				continue;
 			}
-			
+
 			int32 SourcePoseIdx = CacheSourceSkelIndices[SourceBoneName];
-			FTransform SourceBodyRot = (Settings.bTestBodyRotInvariance) ? FTransform::Identity : GetBodyRotation(SourcePhysicsAsset, SourceBoneName);
-			FTransform SourceBodyTfm = SourceBodyRot * GetBodyTransformInvariant(SourcePhysicsAsset, SourceBoneName) * SourceScaleTfm;
-			FTransform SourceBodyCompTfm = SourceBodyTfm * InSourceGlobalPose[SourcePoseIdx];
-			DebugDrawInfo.SourceBodyTfms.Add(SourceBodyCompTfm);
-			
-			int32 TargetPoseIdx = CacheTargetSkelIndices[TargetBoneName];
-			FTransform TargetBodyRot = (Settings.bTestBodyRotInvariance) ? FTransform::Identity :GetBodyRotation(TargetPhysicsAsset, TargetBoneName);
-			FTransform TargetBodyTfm = TargetBodyRot * GetBodyTransformInvariant(TargetPhysicsAsset, TargetBoneName);
-			FTransform TargetBodyCompTfm = TargetBodyTfm * OutTargetGlobalPose[TargetPoseIdx];
-			DebugDrawInfo.TargetBodyTfms.Add(TargetBodyCompTfm);
+			DebugDrawInfo.SourceBodyInfo.Add({SourceBoneName, SourceScaleTfm * InSourceGlobalPose[SourcePoseIdx]});
 		}
 
-		DebugDrawInfo.TargetIntersectTfms.Reset(Settings.IntersectBodies.Num());
-		for (FName BodyName : Settings.IntersectBodies)
+		// Cache debug info for target bodies
+		DebugDrawInfo.TargetBodyInfo.Reset(CacheSourceEffectBones.Num());
+		for (FName SourceBoneName : CacheSourceEffectBones)
 		{
-			if (!CacheTargetSkelIndices.Contains(BodyName))
+			// TODO: Only show domain bodies? 
+			FName TargetBoneName = ApplyBodyMap(SourceBoneName);
+			if (!CacheTargetSkelIndices.Contains(TargetBoneName))
 			{
 				continue;
 			}
 			
-			int32 TargetPoseIdx = CacheTargetSkelIndices[BodyName];
-			DebugDrawInfo.TargetIntersectTfms.Add(OutTargetGlobalPose[TargetPoseIdx]);
+			int32 TargetPoseIdx = CacheTargetSkelIndices[TargetBoneName];
+			DebugDrawInfo.TargetBodyInfo.Add({TargetBoneName, OutTargetGlobalPose[TargetPoseIdx]});
 		}
+
+		// TODO: Intersection test debug info
+		// DebugDrawInfo.TargetIntersectTfms.Reset(Settings.IntersectBodies.Num());
+		// for (FName BodyName : Settings.IntersectBodies)
+		// {
+		// 	if (!CacheTargetSkelIndices.Contains(BodyName))
+		// 	{
+		// 		continue;
+		// 	}
+		// 	
+		// 	int32 TargetPoseIdx = CacheTargetSkelIndices[BodyName];
+		// 	DebugDrawInfo.TargetIntersectTfms.Add(OutTargetGlobalPose[TargetPoseIdx]);
+		// }
 	}
 #endif
 }
@@ -927,174 +936,209 @@ void FIKRetargetRelativeIKOp::DebugDraw(
 	FTransform SourceNoScaleTfm(InSourceTransform.GetRotation(), InSourceTransform.GetTranslation());
 	const FTransform& SourceScaleCheckedTfm = (Settings.bIgnoreSourceScale) ? InSourceTransform : SourceNoScaleTfm;
 
-	if (Settings.bDebugDrawBodyCoords)
+	if (Settings.bDebugDrawBodyPairs)
 	{
-		const float CoordSize = 10.0f;
-		const float CoordThickness = 1.0f;
-		
-		// Show Source Body coords
-		for (const FTransform& BodyTfm : DebugDrawInfo.SourceBodyTfms)
+		// Draw Source body pairs
+		DebugDrawBodyPairs(InPDI, SourceScaleCheckedTfm, SourceScale, DebugDrawInfo.SourcePairVerts);
+
+		// Draw target body pairs (Possibly with "full" retarget space display)
+		DebugDrawBodyPairs(InPDI, InComponentTransform, InComponentScale, DebugDrawInfo.TargetPairVerts);
+	}
+
+	if (Settings.bDebugDrawPhysicsBodies)
+	{
+		if (Settings.SourcePhysicsAssetOverride)
 		{
-			FTransform BodyFrame = BodyTfm * SourceNoScaleTfm;
-			DrawCoordinateSystem(
-				InPDI,
-				BodyFrame.GetLocation(),
-				BodyFrame.GetRotation().Rotator(),
-				static_cast<float>(CoordSize * SourceScale),
-				SDPG_Foreground,
-				static_cast<float>(CoordThickness * SourceScale));
+			DebugDrawBodies(InPDI, SourceNoScaleTfm, SourceScale, Settings.SourcePhysicsAssetOverride, DebugDrawInfo.SourceBodyInfo);
 		}
-
-		// Show Target Body coords
-		for (const FTransform& BodyTfm : DebugDrawInfo.TargetBodyTfms)
+		if (Settings.SourcePhysicsAssetOverride)
 		{
-			FTransform BodyFrame = BodyTfm * InComponentTransform;
-			DrawCoordinateSystem(
-				InPDI,
-				BodyFrame.GetLocation(),
-				BodyFrame.GetRotation().Rotator(),
-				static_cast<float>(CoordSize * InComponentScale),
-				SDPG_Foreground,
-				static_cast<float>(CoordThickness * InComponentScale));
+			DebugDrawBodies(InPDI, InComponentTransform, InComponentScale, Settings.TargetPhysicsAssetOverride, DebugDrawInfo.TargetBodyInfo);
 		}
 	}
 
 	if (Settings.bDebugDrawRetargetVertAverages)
 	{
 		// Show retarget contribution average with offsets
-		for (const FDebugRelativeTargetPairSpace& t : DebugDrawInfo.PairContrib)
-		{
-			FTransform V1Tfm = FTransform(t.PairRangeStart) * InComponentTransform;
-			FTransform V2Tfm = FTransform(t.PairRangeEnd) * InComponentTransform;
-			IKRigDebugRendering::DrawWireCube(InPDI, V1Tfm, FLinearColor::Red,static_cast<float>(0.5*InComponentScale),static_cast<float>(0.5*InComponentScale));
-			IKRigDebugRendering::DrawWireCube(InPDI, V2Tfm, FLinearColor::Red,static_cast<float>(0.5*InComponentScale),static_cast<float>(0.5*InComponentScale));
-
-			FVector V1 = V1Tfm.GetTranslation();
-			FVector V2 = V2Tfm.GetTranslation();
-			// DrawDashedLine(InPDI, VBone, V1, FLinearColor::Yellow, 1, SDPG_Foreground);
-			// DrawDashedLine(InPDI, VBone, V2, FLinearColor::Yellow, 1, SDPG_Foreground);
-			DrawDashedLine(InPDI, V1, V2, FLinearColor::Black, 1, SDPG_Foreground);
-
-			FTransform VAvgTfm = FTransform(t.PairTarget) * InComponentTransform;
-			IKRigDebugRendering::DrawWireCube(InPDI, VAvgTfm, FLinearColor::Gray,static_cast<float>(0.5*InComponentScale),static_cast<float>(0.5*InComponentScale));
-
-			FVector FeasibleMinV =  FMath::Lerp(t.PairRangeStart, t.PairRangeEnd, t.FeasibleRange.Min);
-			FTransform VMinTfm = FTransform(FeasibleMinV) * InComponentTransform;
-			IKRigDebugRendering::DrawWireCube(InPDI, VMinTfm, FLinearColor::Yellow,static_cast<float>(0.5*InComponentScale),static_cast<float>(0.5*InComponentScale));
-
-			FVector FeasibleMaxV =  FMath::Lerp(t.PairRangeStart, t.PairRangeEnd, t.FeasibleRange.Max);
-			FTransform VMaxTfm = FTransform(FeasibleMaxV) * InComponentTransform;
-			IKRigDebugRendering::DrawWireCube(InPDI, VMaxTfm, FLinearColor::Green,static_cast<float>(0.5*InComponentScale),static_cast<float>(0.5*InComponentScale));
-			// DrawDashedLine(InPDI, VOffset, VOffset + t.Get<3>(), FLinearColor::Black, 1, SDPG_Foreground);
-		}
+		DebugDrawPairVertRetarget(InPDI, InComponentTransform, InComponentScale, DebugDrawInfo.PairRetargetInfo);
 	}
 
 	if (Settings.bDebugDrawGoalContributions)
 	{
 		// Show IK Goals (and contribs)
-		for (const FDebugDrawTargetGoal& t : DebugDrawInfo.TargetGoals)
-		{
-			FTransform GoalTfm = FTransform(t.Goal) * InComponentTransform;
-			FVector GoalPoint = GoalTfm.GetTranslation();
-			for (const FVector& v: t.PairTargets)
-			{
-				FTransform ContribTfm = FTransform(v) * InComponentTransform;
-				FVector ContribPt = ContribTfm.GetTranslation();
-				DrawDashedLine(InPDI, ContribPt, GoalPoint, FLinearColor::Red, 1, SDPG_Foreground);
-				IKRigDebugRendering::DrawWireCube(InPDI, ContribTfm, FLinearColor::White,static_cast<float>(InComponentScale),static_cast<float>(InComponentScale));
-			}
-			IKRigDebugRendering::DrawWireCube(InPDI, GoalTfm, FLinearColor::Yellow,static_cast<float>(InComponentScale),static_cast<float>(InComponentScale));
-		}
+		DebugDrawGoalContributions(InPDI, InComponentTransform, InComponentScale, DebugDrawInfo.TargetGoals);
 	}
+}
 
-	if (Settings.bDebugDrawBodyPairs)
+void FIKRetargetRelativeIKOp::DebugDrawBodyPairs(FPrimitiveDrawInterface* InPDI, const FTransform& BaseTransform, double Scale, const TArray<FDebugDrawBodyPair>& BodyPairs) const
+{
+	for (const FDebugDrawBodyPair& Pair : BodyPairs)
 	{
-		// Draw Source body pairs
-		for (int Index = 0; Index < DebugDrawInfo.SourcePairVerts.Num(); ++Index)
+		FTransform TfmBodyA_A = FTransform(Pair.PosA_A) * BaseTransform;
+		FTransform TfmBodyB_B = FTransform(Pair.PosB_B) * BaseTransform;
+
+		IKRigDebugRendering::DrawWireCube(InPDI, TfmBodyA_A, FLinearColor::Blue,static_cast<float>(Scale),static_cast<float>(Scale));
+		IKRigDebugRendering::DrawWireCube(InPDI, TfmBodyB_B, FLinearColor::Red,static_cast<float>(Scale),static_cast<float>(Scale));
+
+		// Possibly draw full retarget debug info
+		if (Pair.bFullPos)
 		{
-			FTransform PrimTfm = FTransform(DebugDrawInfo.SourcePairVerts[Index].Get<0>()) * SourceScaleCheckedTfm;
-			FTransform SecTfm = FTransform(DebugDrawInfo.SourcePairVerts[Index].Get<1>()) * SourceScaleCheckedTfm;
+			FTransform TfmBodyA_B = FTransform(Pair.PosA_B) * BaseTransform;
+			FTransform TfmBodyB_A = FTransform(Pair.PosB_A) * BaseTransform;
 
-			FVector v1 = PrimTfm.GetTranslation();
-			FVector v2 = SecTfm.GetTranslation();
-		
-			IKRigDebugRendering::DrawWireCube(InPDI, PrimTfm, FLinearColor::Blue,static_cast<float>(SourceScale),static_cast<float>(SourceScale));
-			IKRigDebugRendering::DrawWireCube(InPDI, SecTfm, FLinearColor::Red,static_cast<float>(SourceScale),static_cast<float>(SourceScale));
-			DrawDashedLine(InPDI, v1, v2, FLinearColor::Black, 1, SDPG_Foreground);
+			IKRigDebugRendering::DrawWireCube(InPDI, TfmBodyA_B, FLinearColor::Black,static_cast<float>(Scale),static_cast<float>(Scale));
+			IKRigDebugRendering::DrawWireCube(InPDI, TfmBodyB_A, FLinearColor::Yellow,static_cast<float>(Scale),static_cast<float>(Scale));
+
+			// Draw direct relationships between bodies in retargeted A/B spaces
+			DrawDashedLine(InPDI, TfmBodyA_A.GetLocation(), TfmBodyB_A.GetLocation(), FLinearColor::Black, 1, SDPG_Foreground);
+			DrawDashedLine(InPDI, TfmBodyA_B.GetLocation(), TfmBodyB_B.GetLocation(), FLinearColor::Black, 1, SDPG_Foreground);
+
+			// Yellow line showing difference between body A in A-retarget/B-retarget
+			DrawDashedLine(InPDI, TfmBodyA_A.GetLocation(), TfmBodyA_B.GetLocation(), FLinearColor::Yellow, 1, SDPG_Foreground);
+			// Red line showing difference between body B in A-retarget/B-retarget
+			DrawDashedLine(InPDI, TfmBodyB_A.GetLocation(), TfmBodyB_B.GetLocation(), FLinearColor::Red, 1, SDPG_Foreground);
 		}
+		else
+		{
+			DrawDashedLine(InPDI, TfmBodyA_A.GetLocation(), TfmBodyB_B.GetLocation(), FLinearColor::Black, 1, SDPG_Foreground);
+		}
+	}
+}
 
-		// Draw target body pairs (relative to original retarget body)
-		for (int Index = 0; Index < DebugDrawInfo.TargetPairVerts.Num(); Index += 2)
+void FIKRetargetRelativeIKOp::DebugDrawGoalContributions(FPrimitiveDrawInterface* InPDI, const FTransform& BaseTransform, double Scale, const TArray<FDebugDrawTargetGoal>& GoalInfoLis) const
+{
+	// Show IK Goals (and contribs)
+	for (const FDebugDrawTargetGoal& Goal : GoalInfoLis)
+	{
+		FTransform GoalTfm = FTransform(Goal.Goal) * BaseTransform;
+		FVector GoalPoint = GoalTfm.GetTranslation();
+		for (const FVector& v: Goal.PairTargets)
 		{
-			FTransform PrimTfmA = FTransform(DebugDrawInfo.TargetPairVerts[Index].Get<0>()) * InComponentTransform;
-			FTransform SecTfmA = FTransform(DebugDrawInfo.TargetPairVerts[Index].Get<1>()) * InComponentTransform;
+			FTransform ContribTfm = FTransform(v) * BaseTransform;
+			FVector ContribPt = ContribTfm.GetTranslation();
+			DrawDashedLine(InPDI, ContribPt, GoalPoint, FLinearColor::Red, 1, SDPG_Foreground);
+			IKRigDebugRendering::DrawWireCube(InPDI, ContribTfm, FLinearColor::White,static_cast<float>(Scale),static_cast<float>(Scale));
+		}
+		IKRigDebugRendering::DrawWireCube(InPDI, GoalTfm, FLinearColor::Yellow,static_cast<float>(Scale),static_cast<float>(Scale));
+	}
+}
 
-			FTransform PrimTfmB = FTransform(DebugDrawInfo.TargetPairVerts[Index+1].Get<0>()) * InComponentTransform;
-			FTransform SecTfmB = FTransform(DebugDrawInfo.TargetPairVerts[Index+1].Get<1>()) * InComponentTransform;
+void FIKRetargetRelativeIKOp::DebugDrawPairVertRetarget(FPrimitiveDrawInterface* InPDI, const FTransform& BaseTransform, double Scale, const TArray<FDebugRelativeTargetPairSpace>& RetargetPairList) const
+{
+	const float MarkerSize = 0.5f;
+	for (const FDebugRelativeTargetPairSpace& Info : RetargetPairList)
+	{
+		FTransform V1Tfm = FTransform(Info.PairRangeStart) * BaseTransform;
+		FTransform V2Tfm = FTransform(Info.PairRangeEnd) * BaseTransform;
+		IKRigDebugRendering::DrawWireCube(InPDI, V1Tfm, FLinearColor::Red,static_cast<float>(MarkerSize*Scale),static_cast<float>(MarkerSize*Scale));
+		IKRigDebugRendering::DrawWireCube(InPDI, V2Tfm, FLinearColor::Red,static_cast<float>(MarkerSize*Scale),static_cast<float>(MarkerSize*Scale));
+
+		FVector V1 = V1Tfm.GetTranslation();
+		FVector V2 = V2Tfm.GetTranslation();
+		// DrawDashedLine(InPDI, VBone, V1, FLinearColor::Yellow, 1, SDPG_Foreground);
+		// DrawDashedLine(InPDI, VBone, V2, FLinearColor::Yellow, 1, SDPG_Foreground);
+		DrawDashedLine(InPDI, V1, V2, FLinearColor::Black, 1, SDPG_Foreground);
+
+		FTransform VAvgTfm = FTransform(Info.PairTarget) * BaseTransform;
+		IKRigDebugRendering::DrawWireCube(InPDI, VAvgTfm, FLinearColor::Gray,static_cast<float>(MarkerSize*Scale),static_cast<float>(MarkerSize*Scale));
+
+		FVector FeasibleMinV =  FMath::Lerp(Info.PairRangeStart, Info.PairRangeEnd, Info.FeasibleRange.Min);
+		FTransform VMinTfm = FTransform(FeasibleMinV) * BaseTransform;
+		IKRigDebugRendering::DrawWireCube(InPDI, VMinTfm, FLinearColor::Yellow,static_cast<float>(MarkerSize*Scale),static_cast<float>(MarkerSize*Scale));
+
+		FVector FeasibleMaxV =  FMath::Lerp(Info.PairRangeStart, Info.PairRangeEnd, Info.FeasibleRange.Max);
+		FTransform VMaxTfm = FTransform(FeasibleMaxV) * BaseTransform;
+		IKRigDebugRendering::DrawWireCube(InPDI, VMaxTfm, FLinearColor::Green,static_cast<float>(MarkerSize*Scale),static_cast<float>(MarkerSize*Scale));
+		// DrawDashedLine(InPDI, VOffset, VOffset + t.Get<3>(), FLinearColor::Black, 1, SDPG_Foreground);
+	}
+}
 
-			FVector v1 = PrimTfmA.GetTranslation();
-			FVector v2 = SecTfmA.GetTranslation();
-			FVector v3 = PrimTfmB.GetTranslation();
-			FVector v4 = SecTfmB.GetTranslation();
+void FIKRetargetRelativeIKOp::DebugDrawBodies(FPrimitiveDrawInterface* InPDI, const FTransform& BaseTransform, double Scale, UPhysicsAsset* PhysAsset, const TArray<FDebugDrawBodyInfo>& PhysBodies) const
+{
+	const float CoordSize = 10.0f;
+	const float CoordThickness = 0.5f;
+	
+	for (const FDebugDrawBodyInfo& BodyInfo : PhysBodies)
+	{
+		const FKShapeElem* ShapeElem = FindBodyShape(PhysAsset, BodyInfo.BodyName);
+		if (!ShapeElem)
+		{
+			continue;
+		}
 		
-			IKRigDebugRendering::DrawWireCube(InPDI, PrimTfmA, FLinearColor::Blue,static_cast<float>(InComponentScale),static_cast<float>(InComponentScale));
-			IKRigDebugRendering::DrawWireCube(InPDI, SecTfmA, FLinearColor::Red,static_cast<float>(InComponentScale),static_cast<float>(InComponentScale));
+		FTransform CompTfm = BodyInfo.BoneTfm * BaseTransform;
 
-			IKRigDebugRendering::DrawWireCube(InPDI, PrimTfmB, FLinearColor::Black,static_cast<float>(InComponentScale),static_cast<float>(InComponentScale));
-			IKRigDebugRendering::DrawWireCube(InPDI, SecTfmB, FLinearColor::Yellow,static_cast<float>(InComponentScale),static_cast<float>(InComponentScale));
-			
-			DrawDashedLine(InPDI, v1, v3, FLinearColor::Yellow, 1, SDPG_Foreground);
-			DrawDashedLine(InPDI, v2, v3, FLinearColor::Black, 1, SDPG_Foreground);
-			DrawDashedLine(InPDI, v2, v4, FLinearColor::Red, 1, SDPG_Foreground);
-			DrawDashedLine(InPDI, v4, v1, FLinearColor::Black, 1, SDPG_Foreground);
-		}
+		FTransform BodyRot = (Settings.bTestBodyRotInvariance) ? FTransform::Identity : GetBodyRotation(SourcePhysicsAsset, BodyInfo.BodyName);
+		FTransform BodyTfm = BodyRot * GetBodyTransformInvariant(PhysAsset, BodyInfo.BodyName) * CompTfm;
+
+		// Draw body coord system:
+		DrawCoordinateSystem(
+			InPDI,
+			BodyTfm.GetLocation(),
+			BodyTfm.GetRotation().Rotator(),
+			static_cast<float>(CoordSize * Scale),
+			SDPG_Foreground,
+			static_cast<float>(CoordThickness * Scale));
+
+
+		// Draw physics body
+		DebugDrawPhysBody(InPDI, CompTfm, Scale, ShapeElem, FLinearColor::Yellow);
 	}
+}
 
-	// if (Settings.bDebugDrawIntersectTest)
-	// {
-	// 	for (int32 BodyIdx = 0; BodyIdx < DebugDrawInfo.TargetIntersectTfms.Num(); ++BodyIdx)
-	// 	{
-	// 		if (!Settings.IntersectBodies.IsValidIndex(BodyIdx))
-	// 		{
-	// 			continue;
-	// 		}
-	// 		
-	// 		FName BodyName = Settings.IntersectBodies[BodyIdx];
-	// 		FKShapeElem* ShapeElem = FindBodyShape(TargetPhysicsAsset, BodyName);
-	// 		if (!ShapeElem || ShapeElem->GetShapeType() != EAggCollisionShape::Sphyl)
-	// 		{
-	// 			continue;
-	// 		}
-	//
-	// 		const FKSphylElem* CapsuleElem = static_cast<FKSphylElem*>(ShapeElem);
-	// 		FTransform Tfm = CapsuleElem->GetTransform() * DebugDrawInfo.TargetIntersectTfms[BodyIdx] * InComponentTransform;
-	// 		const FVector Translation = Tfm.GetTranslation();
-	// 		const FVector UnitXAxis   = Tfm.GetUnitAxis( EAxis::X );
-	// 		const FVector UnitYAxis   = Tfm.GetUnitAxis( EAxis::Y );
-	// 		const FVector UnitZAxis   = Tfm.GetUnitAxis( EAxis::Z );
-	//
-	// 		DrawWireCapsule(InPDI, Translation,
-	// 			UnitXAxis, UnitYAxis, UnitZAxis,
-	// 			FLinearColor::Yellow,
-	// 			CapsuleElem->Radius * InComponentScale,
-	// 			(0.5*CapsuleElem->Length + CapsuleElem->Radius) * InComponentScale,
-	// 			25, SDPG_Foreground, 0, 1.0);
-	// 	}
-	// }
+void FIKRetargetRelativeIKOp::DebugDrawPhysBody(FPrimitiveDrawInterface* InPDI, const FTransform& ParentTransform, double Scale, const FKShapeElem* ShapeElem, const FLinearColor& Color) const
+{
+	FTransform BodyFrame = ShapeElem->GetTransform() * ParentTransform;
+
+	const FVector Translation = BodyFrame.GetLocation();
+	const FVector UnitXAxis = BodyFrame.GetUnitAxis( EAxis::X );
+	const FVector UnitYAxis = BodyFrame.GetUnitAxis( EAxis::Y );
+	const FVector UnitZAxis = BodyFrame.GetUnitAxis( EAxis::Z );
+	
+	switch (ShapeElem->GetShapeType())
+	{
+	case EAggCollisionShape::Box:
+		{
+			const FKBoxElem* BoxElem = static_cast<const FKBoxElem*>(ShapeElem);
+			const FVector Extent = 0.5 * FVector(BoxElem->X, BoxElem->Y, BoxElem->Z) * Scale;
+			DrawOrientedWireBox(InPDI, Translation, UnitXAxis, UnitYAxis, UnitZAxis, Extent, Color, SDPG_Foreground);
+			return;
+		}
+	case EAggCollisionShape::Sphyl:
+		{
+			const FKSphylElem* CapsuleElem = static_cast<const FKSphylElem*>(ShapeElem);
+			const double Radius = CapsuleElem->Radius * Scale;
+			const double HalfHeight = (0.5*CapsuleElem->Length + CapsuleElem->Radius) * Scale;
+			DrawWireCapsule(InPDI, Translation, UnitXAxis, UnitYAxis, UnitZAxis, Color,
+						Radius, HalfHeight,25, SDPG_Foreground, 0, 1.0);
+			return;
+		}
+	case EAggCollisionShape::Sphere:
+		{
+			const FKSphereElem* SphereElem = static_cast<const FKSphereElem*>(ShapeElem);
+			const double Radius = SphereElem->Radius * Scale;
+			DrawWireSphere(InPDI, Translation, Color, Radius, 25, SDPG_Foreground);
+			return;
+		}
+	default:
+		return;
+	}
 }
 
+
+
 void FIKRetargetRelativeIKOp::ResetDebugInfo()
 {
 	DebugDrawInfo.SourcePairVerts.Reset();
 	DebugDrawInfo.TargetPairVerts.Reset();
 	
-	DebugDrawInfo.SourceBodyTfms.Reset();
-	DebugDrawInfo.TargetBodyTfms.Reset();
+	DebugDrawInfo.SourceBodyInfo.Reset();
+	DebugDrawInfo.TargetBodyInfo.Reset();
 	
 	DebugDrawInfo.TargetIntersectTfms.Reset();
 	DebugDrawInfo.TestSpheres.Reset();
 	
-	DebugDrawInfo.PairContrib.Reset();
+	DebugDrawInfo.PairRetargetInfo.Reset();
 	DebugDrawInfo.TargetGoals.Reset();
 }
 
diff --git a/Engine/Plugins/Experimental/Animation/RelativeIKOp/Source/RelativeIKOp/Public/RelativeIKOp.h b/Engine/Plugins/Experimental/Animation/RelativeIKOp/Source/RelativeIKOp/Public/RelativeIKOp.h
--- a/Engine/Plugins/Experimental/Animation/RelativeIKOp/Source/RelativeIKOp/Public/RelativeIKOp.h
+++ b/Engine/Plugins/Experimental/Animation/RelativeIKOp/Source/RelativeIKOp/Public/RelativeIKOp.h
@@ -35,23 +35,41 @@ struct FDebugDrawTargetGoal
 	TArray<FVector> PairTargets;
 };
 
+// Debug draw info for source or target pair
+struct FDebugDrawBodyPair
+{
+	FVector PosA_A;
+	FVector PosB_B;
+	
+	bool bFullPos;
+	FVector PosB_A;
+	FVector PosA_B;
+};
+
+// Debug physics body info
+struct FDebugDrawBodyInfo
+{
+	FName BodyName;
+	FTransform BoneTfm;
+};
+
 // Used to hold all debug draw info for convenience
 struct FDebugRelativeIKDrawInfo
 {
 	// Source/target vert pair info
-	TArray<TTuple<FVector,FVector>> SourcePairVerts;
-	TArray<TTuple<FVector,FVector>> TargetPairVerts;
+	TArray<FDebugDrawBodyPair> SourcePairVerts;
+	TArray<FDebugDrawBodyPair> TargetPairVerts;
 
-	// Body coordinates
-	TArray<FTransform> SourceBodyTfms;
-	TArray<FTransform> TargetBodyTfms;
+	// Body info
+	TArray<FDebugDrawBodyInfo> SourceBodyInfo;
+	TArray<FDebugDrawBodyInfo> TargetBodyInfo;
 
 	// Intersect info
 	TArray<FTransform> TargetIntersectTfms;
 	TArray<TTuple<FVector,double>> TestSpheres;
 
-	// Debug info for pair contributions
-	TArray<FDebugRelativeTargetPairSpace> PairContrib;
+	// Debug info for retarget pairs
+	TArray<FDebugRelativeTargetPairSpace> PairRetargetInfo;
 	// Final goal points with contributions
 	TArray<FDebugDrawTargetGoal> TargetGoals;
 };
@@ -124,7 +142,7 @@ struct FIKRetargetRelativeIKOpSettings : public FIKRetargetOpSettingsBase
 	bool bDebugDrawBodyPairs = false;
 	
 	UPROPERTY(EditAnywhere, BlueprintReadWrite, Category = "Debug")
-	TArray<FName> DrawTargetPairBones;
+	TArray<FName> DebugFullRetargetPairBones;
 
 	UPROPERTY(EditAnywhere, BlueprintReadWrite, Category = "Debug")
 	bool bDebugDrawGoalContributions = false;
@@ -133,16 +151,7 @@ struct FIKRetargetRelativeIKOpSettings : public FIKRetargetOpSettingsBase
 	bool bDebugDrawRetargetVertAverages = false;
 
 	UPROPERTY(EditAnywhere, BlueprintReadWrite, Category = "Debug")
-	bool bDebugDrawBodyCoords = false;
-
-	// UPROPERTY(EditAnywhere, BlueprintReadWrite, Category = "Debug")
-	// bool bDebugDrawDomainBodies = false;
-	//
-	// UPROPERTY(EditAnywhere, BlueprintReadWrite, Category = "Debug")
-	// bool bDebugDrawContactBodies = false;
-	//
-	// UPROPERTY(EditAnywhere, BlueprintReadWrite, Category = "Debug")
-	// bool bDebugDrawIntersectTest = false;
+	bool bDebugDrawPhysicsBodies = false;
 
 	UPROPERTY(EditAnywhere, BlueprintReadWrite, Category = "Test Features")
 	bool bTestDistContactAlpha = true;
@@ -228,6 +237,12 @@ public:
 		const FIKRetargetDebugDrawState& InEditorState) const override;
 
 private:
+	void DebugDrawBodyPairs(FPrimitiveDrawInterface* InPDI, const FTransform& BaseTransform, double Scale, const TArray<FDebugDrawBodyPair>& BodyPairs) const;
+	void DebugDrawGoalContributions(FPrimitiveDrawInterface* InPDI, const FTransform& BaseTransform, double Scale, const TArray<FDebugDrawTargetGoal>& GoalInfoLis) const;
+	void DebugDrawPairVertRetarget(FPrimitiveDrawInterface* InPDI, const FTransform& BaseTransform, double Scale, const TArray<FDebugRelativeTargetPairSpace>& RetargetPairList) const;
+	void DebugDrawBodies(FPrimitiveDrawInterface* InPDI, const FTransform& BaseTransform, double Scale, UPhysicsAsset* PhysAsset, const TArray<FDebugDrawBodyInfo>& PhysBodies) const;
+	void DebugDrawPhysBody(FPrimitiveDrawInterface* InPDI, const FTransform& ParentTransform, double Scale, const FKShapeElem* ShapeElem, const FLinearColor& Color) const;
+	
 	void ResetDebugInfo();
 
 	// Collection of all debug structures
