RelativeIKOp: Added temporal smoothing and intersection sphere rendering

diff --git a/Engine/Plugins/Experimental/Animation/RelativeIKOp/Source/RelativeIKOp/Private/RelativeIKOp.cpp b/Engine/Plugins/Experimental/Animation/RelativeIKOp/Source/RelativeIKOp/Private/RelativeIKOp.cpp
--- a/Engine/Plugins/Experimental/Animation/RelativeIKOp/Source/RelativeIKOp/Private/RelativeIKOp.cpp
+++ b/Engine/Plugins/Experimental/Animation/RelativeIKOp/Source/RelativeIKOp/Private/RelativeIKOp.cpp
@@ -113,6 +113,22 @@ void FIKRetargetRelativeIKOp::Run(
 	int32 SampleIdx = FMath::Min(static_cast<int32>(NotifyTime * SampleRate), CachedNotifyInfo->NumSamples-1);
 	int32 FramePairIdx = NumBodies * SampleIdx;
 
+	// Testing temporal smoothing of local points
+	const int32 MinSmoothIdx = FMath::Max(SampleIdx - Settings.TestTemporalSmoothingFrames, 0);
+	const int32 MaxSmoothIdx = FMath::Min(SampleIdx + Settings.TestTemporalSmoothingFrames,CachedNotifyInfo->NumSamples-1) ;
+
+	TArray<FVector3f> SmoothedPoints;
+	SmoothedPoints.SetNum(NumBodies);
+	for (int BodyIdx = 0; BodyIdx < CachedNotifyInfo->BodyPairs.Num(); ++BodyIdx)
+	{
+		SmoothedPoints[BodyIdx] = FVector3f::Zero();
+		for (int32 FrameIdx = MinSmoothIdx; FrameIdx <= MaxSmoothIdx; ++FrameIdx)
+		{
+			SmoothedPoints[BodyIdx] += CachedNotifyInfo->BodyPairsLocalReference[NumBodies * FrameIdx + BodyIdx];
+		}
+		SmoothedPoints[BodyIdx] /= static_cast<float>(MaxSmoothIdx - MinSmoothIdx + 1);
+	}
+
 	FDebugRelativeIKDrawInfo LocalDebugInfo;
 	
 #if WITH_EDITOR
@@ -172,14 +188,16 @@ void FIKRetargetRelativeIKOp::Run(
 		FTransform TargetOtherBodyRot = (Settings.bTestBodyRotInvariance) ? FTransform::Identity : SourceOtherBodyRot.Inverse() * GetBodyRotation(TargetPhysicsAsset, TargetOtherBone);
 		
 		// Compute source transforms and source-space vertex distance
-		const FVector3f& LocalPrimaryVert = CachedNotifyInfo->BodyPairsLocalReference[FramePairIdx + 2*PairIdx];
+		// const FVector3f& LocalPrimaryVert = CachedNotifyInfo->BodyPairsLocalReference[FramePairIdx + 2*PairIdx];
+		const FVector3f& LocalPrimaryVert = SmoothedPoints[2*PairIdx];
 		// TODO: Swap and clean these up if we bake in rot invaraint space
 		// FTransform SourcePrimaryBodyTfm = GetBodyTransformInvariant(SourcePhysicsAsset, SourcePrimaryBone) * SourceScaleTfm;
 		FTransform SourcePrimaryBodyTfm = SourcePrimaryBodyRot * GetBodyTransformInvariant(SourcePhysicsAsset, SourcePrimaryBone) * SourceScaleTfm;
 		FTransform SourcePrimaryCompTfm = SourcePrimaryBodyTfm * InSourceGlobalPose[SourcePrimaryPoseIdx] * InvScaleTfm;
 		FVector SourcePrimaryCompV = SourcePrimaryCompTfm.TransformPosition(FVector(LocalPrimaryVert));
 
-		const FVector3f& LocalOtherVert = CachedNotifyInfo->BodyPairsLocalReference[FramePairIdx + 2*PairIdx + 1];
+		// const FVector3f& LocalOtherVert = CachedNotifyInfo->BodyPairsLocalReference[FramePairIdx + 2*PairIdx + 1];
+		const FVector3f& LocalOtherVert = SmoothedPoints[2*PairIdx+1];
 		FTransform SourceOtherBodyTfm = SourceOtherBodyRot * GetBodyTransformInvariant(SourcePhysicsAsset, SourceOtherBone) * SourceScaleTfm;
 		FTransform SourceOtherCompTfm = SourceOtherBodyTfm * InSourceGlobalPose[SourceOtherPoseIdx] * InvScaleTfm;
 		FVector SourceOtherCompV = SourceOtherCompTfm.TransformPosition(FVector(LocalOtherVert));
@@ -329,10 +347,18 @@ void FIKRetargetRelativeIKOp::Run(
 			
 			// TODO: Play with total weight functions to alpha goal
 			float GoalAlpha = FMath::Min(1.0f, static_cast<float>(TotalWeight*FMath::Clamp(Settings.ContributionSumWeight, 0.0f, 1.0f)));
-			IKGoal->bEnabled = true;
-			IKGoal->Position = GoalPoint;
-			IKGoal->PositionAlpha = GoalAlpha;
-			IKGoal->PositionSpace = EIKRigGoalSpace::Component;
+			// TODO: Blend from old goal value (or target bone)
+			int32 TargetBoneIdx = CacheTargetSkelIndices[TargetBoneName];
+			FVector BlendGoal = FMath::Lerp(OutTargetGlobalPose[TargetBoneIdx].GetLocation(), GoalPoint, GoalAlpha);
+
+			bool bSetGoals = !Settings.bEnableDebugDraw || !Settings.bDryRun;
+			if (bSetGoals)
+			{
+				IKGoal->bEnabled = true;
+				IKGoal->Position = BlendGoal;
+				IKGoal->PositionAlpha = 1.0;
+				IKGoal->PositionSpace = EIKRigGoalSpace::Component;
+			}
 
 			// Sphere-capsule intersection goal deltas
 			if (Settings.bTestIntersectionFixup)
@@ -371,14 +397,15 @@ void FIKRetargetRelativeIKOp::Run(
 						IKGoal->Position += Dist * DeltaDir;
 					}
 				}
-
+			}
 #if WITH_EDITOR
-				if (Settings.bEnableDebugDraw)
-				{
-					LocalDebugInfo.TestSpheres.Emplace(IKGoal->Position, BoneRadius);
-				}
-#endif //WITH_EDITOR
+			if (Settings.bEnableDebugDraw)
+			{
+				FKShapeElem* BoneShape = FindBodyShape(TargetPhysicsAsset, TargetBoneName);
+				double BoneRadius = CalcShapeSmallRadius(BoneShape);
+				LocalDebugInfo.TestSpheres.Emplace(IKGoal->Position, BoneRadius);
 			}
+#endif //WITH_EDITOR
 		}
 	}
 
@@ -406,9 +433,9 @@ void FIKRetargetRelativeIKOp::Run(
 			DebugDrawInfo.SourceBodyInfo.Add({SourceBoneName, SourceScaleTfm * InSourceGlobalPose[SourcePoseIdx]});
 		}
 
-		// Cache debug info for target bodies
-		DebugDrawInfo.TargetBodyInfo.Reset(CacheSourceEffectBones.Num());
-		for (FName SourceBoneName : CacheSourceEffectBones)
+		// Cache debug info for target domain bodies
+		DebugDrawInfo.TargetBodyInfo.Reset(CacheSourceDomainBones.Num());
+		for (FName SourceBoneName : CacheSourceDomainBones)
 		{
 			// TODO: Only show domain bodies? 
 			FName TargetBoneName = ApplyBodyMap(SourceBoneName);
@@ -420,19 +447,6 @@ void FIKRetargetRelativeIKOp::Run(
 			int32 TargetPoseIdx = CacheTargetSkelIndices[TargetBoneName];
 			DebugDrawInfo.TargetBodyInfo.Add({TargetBoneName, OutTargetGlobalPose[TargetPoseIdx]});
 		}
-
-		// TODO: Intersection test debug info
-		// DebugDrawInfo.TargetIntersectTfms.Reset(Settings.IntersectBodies.Num());
-		// for (FName BodyName : Settings.IntersectBodies)
-		// {
-		// 	if (!CacheTargetSkelIndices.Contains(BodyName))
-		// 	{
-		// 		continue;
-		// 	}
-		// 	
-		// 	int32 TargetPoseIdx = CacheTargetSkelIndices[BodyName];
-		// 	DebugDrawInfo.TargetIntersectTfms.Add(OutTargetGlobalPose[TargetPoseIdx]);
-		// }
 	}
 #endif
 }
@@ -495,10 +509,13 @@ void FIKRetargetRelativeIKOp::ComputeFeasibilityRange(FDoubleInterval& OutFeasib
 	bool UpperRange = ChainDistStart > ChainDistEnd; 
 	if (FMath::Max(ChainDistStart, ChainDistEnd) <= CheckRadius)
 	{
+		// Whole range is feasible
 		OutFeasibleRange = FDoubleInterval(0.0,1.0);
 	}
 	else if (FMath::Min(ChainDistStart, ChainDistEnd) >= CheckRadius)
 	{
+		// None of the range is feasible so use lower-dist endpoint
+		// TODO: This is an approximation as we could have a smaller feasible range that doesn't include either endpoint
 		double RangePt = (UpperRange) ? 1.0 : 0.0;
 		OutFeasibleRange = FDoubleInterval(RangePt,RangePt);
 	}
@@ -516,6 +533,7 @@ void FIKRetargetRelativeIKOp::ComputeFeasibilityRange(FDoubleInterval& OutFeasib
 double FIKRetargetRelativeIKOp::CalcLineSphereIntersect(const FVector& Center, double Radius, const FVector& StartPoint, const FVector& EndPoint)
 {
 	// NOTE: This function assumes two intersection points and returns positive ratio along line from start point
+	// Used for getting the endpoint of feasible range
 	FVector LineVec = EndPoint - StartPoint;
 	FVector LocPt = StartPoint - Center;
 
@@ -551,6 +569,13 @@ double FIKRetargetRelativeIKOp::CalcIntersectionDelta(const FVector& TargetPoint
 			break;
 		}
 
+	// case EAggCollisionShape::Box:
+	// 	{
+	// 		const FKBoxElem* BoxElem = static_cast<const FKBoxElem*>(ShapeElem);
+	// 		Dist = BoxElem->GetClosestPointAndNormal(TargetPoint, ShapeTfm, ClosestPoint, OutDeltaDir);
+	// 		break;
+	// 	}
+
 	default:
 		{
 			return -1.0;
@@ -830,6 +855,7 @@ void FIKRetargetRelativeIKOp::UpdateCacheNotifyInfo(URelativeBodyBakeAnimNotify*
 	// Clear all cached data on update
 	CacheSourceBodyEffectVertIdx.Reset();
 	CacheSourceEffectBones.Reset();
+	CacheSourceDomainBones.Reset();
 	FrameBoneTargets.Reset();
 	FrameBoneWeights.Reset();
 
@@ -883,6 +909,14 @@ void FIKRetargetRelativeIKOp::UpdateCacheNotifyInfo(URelativeBodyBakeAnimNotify*
 			CacheSourceBodyEffectVertIdx[PrimaryIdx].Add(2*PairIdx);
 		}
 	}
+
+	for (int EffectIdx = 0; EffectIdx < CacheSourceEffectBones.Num(); ++EffectIdx)
+	{
+		if (CacheSourceBodyEffectVertIdx[EffectIdx].IsEmpty())
+		{
+			CacheSourceDomainBones.Add(CacheSourceEffectBones[EffectIdx]);
+		}
+	}
 }
 
 
@@ -968,6 +1002,15 @@ void FIKRetargetRelativeIKOp::DebugDraw(
 		// Show IK Goals (and contribs)
 		DebugDrawGoalContributions(InPDI, InComponentTransform, InComponentScale, DebugDrawInfo.TargetGoals);
 	}
+
+	if (Settings.bDebugDrawIntersectionSpheres)
+	{
+		for ( const auto& Sphere : DebugDrawInfo.TestSpheres )
+		{
+			FVector Center = InComponentTransform.TransformVector(Sphere.Get<0>());
+			DrawWireSphere(InPDI, Center, FLinearColor::Red, Sphere.Get<1>()*InComponentScale, 25, SDPG_Foreground);
+		}
+	}
 }
 
 void FIKRetargetRelativeIKOp::DebugDrawBodyPairs(FPrimitiveDrawInterface* InPDI, const FTransform& BaseTransform, double Scale, const TArray<FDebugDrawBodyPair>& BodyPairs) const
@@ -1135,7 +1178,6 @@ void FIKRetargetRelativeIKOp::ResetDebugInfo()
 	DebugDrawInfo.SourceBodyInfo.Reset();
 	DebugDrawInfo.TargetBodyInfo.Reset();
 	
-	DebugDrawInfo.TargetIntersectTfms.Reset();
 	DebugDrawInfo.TestSpheres.Reset();
 	
 	DebugDrawInfo.PairRetargetInfo.Reset();
diff --git a/Engine/Plugins/Experimental/Animation/RelativeIKOp/Source/RelativeIKOp/Public/RelativeIKOp.h b/Engine/Plugins/Experimental/Animation/RelativeIKOp/Source/RelativeIKOp/Public/RelativeIKOp.h
--- a/Engine/Plugins/Experimental/Animation/RelativeIKOp/Source/RelativeIKOp/Public/RelativeIKOp.h
+++ b/Engine/Plugins/Experimental/Animation/RelativeIKOp/Source/RelativeIKOp/Public/RelativeIKOp.h
@@ -65,7 +65,7 @@ struct FDebugRelativeIKDrawInfo
 	TArray<FDebugDrawBodyInfo> TargetBodyInfo;
 
 	// Intersect info
-	TArray<FTransform> TargetIntersectTfms;
+	// TArray<FTransform> TargetIntersectTfms;
 	TArray<TTuple<FVector,double>> TestSpheres;
 
 	// Debug info for retarget pairs
@@ -135,35 +135,57 @@ struct FIKRetargetRelativeIKOpSettings : public FIKRetargetOpSettingsBase
 	UPROPERTY(BlueprintReadWrite, Category = "Parameters")
 	double RetargetSpringAlpha = 0.5;
 
+	// Enable debug displays
 	UPROPERTY(EditAnywhere, BlueprintReadWrite, Category = "Debug")
 	bool bEnableDebugDraw = false;
-	
+
+	// Draw source and target body pair relationships for current animation
 	UPROPERTY(EditAnywhere, BlueprintReadWrite, Category = "Debug")
 	bool bDebugDrawBodyPairs = false;
-	
+
+	// Draw full retarget space (quad) pair relationships
 	UPROPERTY(EditAnywhere, BlueprintReadWrite, Category = "Debug")
 	TArray<FName> DebugFullRetargetPairBones;
 
+	// Draw each pair's goal contribution (white) and show weighted final goal location (yellow)
 	UPROPERTY(EditAnywhere, BlueprintReadWrite, Category = "Debug")
 	bool bDebugDrawGoalContributions = false;
 
+	// Draw retarget pair-line for each pair contribution 
 	UPROPERTY(EditAnywhere, BlueprintReadWrite, Category = "Debug")
 	bool bDebugDrawRetargetVertAverages = false;
 
+	// Display source and target physics bodies for baked data
 	UPROPERTY(EditAnywhere, BlueprintReadWrite, Category = "Debug")
 	bool bDebugDrawPhysicsBodies = false;
 
+	// Display source and target physics bodies for baked data
+	UPROPERTY(EditAnywhere, BlueprintReadWrite, Category = "Debug")
+	bool bDebugDrawIntersectionSpheres = false;
+
+	// Run op and display debug info but DON'T update IK Goals
+	UPROPERTY(EditAnywhere, BlueprintReadWrite, Category = "Debug")
+	bool bDryRun = false;
+
+	// Use distance to pair contact to blend towards direct contact
 	UPROPERTY(EditAnywhere, BlueprintReadWrite, Category = "Test Features")
 	bool bTestDistContactAlpha = true;
 
+	// Dial back distance weighting based on feasibility length of target bone chain
 	UPROPERTY(EditAnywhere, BlueprintReadWrite, Category = "Test Features")
 	bool bTestFeasibilityWeight = true;
 
+	// Move goals to avoid intersection with bodies in Intersect Bodies list
 	UPROPERTY(EditAnywhere, BlueprintReadWrite, Category = "Test Features")
 	bool bTestIntersectionFixup = true;
 
+	// Use rotation-invariant space for physics body components on source/target
 	UPROPERTY(EditAnywhere, BlueprintReadWrite, Category = "Test Features")
 	bool bTestBodyRotInvariance = true;
+
+	// Frames of temporal smoothing of body pair verts (0 is no smoothing)
+	UPROPERTY(EditAnywhere, BlueprintReadWrite, Category = "Test Features", meta=(ClampMin = "0", ClampMax = "7"))
+	int32 TestTemporalSmoothingFrames = 3;
 };
 
 USTRUCT(BlueprintType, meta = (DisplayName = "Relative IK Goal Op"))
@@ -270,6 +292,7 @@ private:
 	// Array of body-vert indices affecting each joint goal (for weighted averaging)
 	TArray<TArray<int32>> CacheSourceBodyEffectVertIdx;
 	TArray<FName> CacheSourceEffectBones;
+	TArray<FName> CacheSourceDomainBones;
 	TArray<double> FrameBoneWeights;
 	TArray<FVector> FrameBoneTargets;
 
