RelativeIKOp: Use body offset for sphere intersection tests.

diff --git a/Engine/Plugins/Experimental/Animation/RelativeIKOp/Source/RelativeIKOp/Private/RelativeIKOp.cpp b/Engine/Plugins/Experimental/Animation/RelativeIKOp/Source/RelativeIKOp/Private/RelativeIKOp.cpp
--- a/Engine/Plugins/Experimental/Animation/RelativeIKOp/Source/RelativeIKOp/Private/RelativeIKOp.cpp
+++ b/Engine/Plugins/Experimental/Animation/RelativeIKOp/Source/RelativeIKOp/Private/RelativeIKOp.cpp
@@ -363,9 +363,12 @@ void FIKRetargetRelativeIKOp::Run(
 			// Sphere-capsule intersection goal deltas
 			if (Settings.bTestIntersectionFixup)
 			{
-				FKShapeElem* BoneShape = FindBodyShape(TargetPhysicsAsset, TargetBoneName);
-				double BoneRadius = CalcShapeSmallRadius(BoneShape);
-				if (FMath::IsNearlyZero(BoneRadius))
+				FKShapeElem* BodyShape = FindBodyShape(TargetPhysicsAsset, TargetBoneName);
+				double BodyRadius = CalcShapeSmallRadius(BodyShape);
+				FVector LocalOffset = BodyShape->GetTransform().GetLocation();
+				FVector BodyOffset = OutTargetGlobalPose[TargetBoneIdx].TransformVector(LocalOffset);
+				FVector BodyLocation = IKGoal->Position + BodyOffset;
+				if (FMath::IsNearlyZero(BodyRadius))
 				{
 					continue;
 				}
@@ -391,19 +394,22 @@ void FIKRetargetRelativeIKOp::Run(
 					}
 
 					FVector DeltaDir = FVector::ZeroVector;
-					double Dist = CalcIntersectionDelta(GoalPoint, BoneRadius, BodyTfm, ShapeElem, DeltaDir);
+					double Dist = CalcIntersectionDelta(BodyLocation, BodyRadius, BodyTfm, ShapeElem, DeltaDir);
 					if (Dist > 0.0)
 					{
-						IKGoal->Position += Dist * DeltaDir;
+						BodyLocation += Dist * DeltaDir;
 					}
 				}
+				IKGoal->Position = BodyLocation - BodyOffset;
 			}
 #if WITH_EDITOR
 			if (Settings.bEnableDebugDraw)
 			{
-				FKShapeElem* BoneShape = FindBodyShape(TargetPhysicsAsset, TargetBoneName);
-				double BoneRadius = CalcShapeSmallRadius(BoneShape);
-				LocalDebugInfo.TestSpheres.Emplace(IKGoal->Position, BoneRadius);
+				FKShapeElem* BodyShape = FindBodyShape(TargetPhysicsAsset, TargetBoneName);
+				double BodyRadius = CalcShapeSmallRadius(BodyShape);
+				FVector LocalOffset = BodyShape->GetTransform().GetLocation();
+				FVector BodyOffset = OutTargetGlobalPose[TargetBoneIdx].TransformVector(LocalOffset);
+				LocalDebugInfo.TestSpheres.Emplace(IKGoal->Position+BodyOffset, BodyRadius);
 			}
 #endif //WITH_EDITOR
 		}
