RelativeIK: Split relative IK and intersection test ops.

diff --git a/Engine/Plugins/Experimental/Animation/RelativeIKOp/RelativeIKOp.uplugin b/Engine/Plugins/Experimental/Animation/RelativeIKOp/RelativeIKOp.uplugin
--- a/Engine/Plugins/Experimental/Animation/RelativeIKOp/RelativeIKOp.uplugin
+++ b/Engine/Plugins/Experimental/Animation/RelativeIKOp/RelativeIKOp.uplugin
@@ -25,6 +25,11 @@
 			"Type": "Runtime",
 			"LoadingPhase": "Default"
 		},
+		{
+			"Name": "BodyIntersectIKOp",
+			"Type": "Runtime",
+			"LoadingPhase": "Default"
+		},
 		{
 			"Name": "RelativeBodyAnimUtils",
 			"Type": "Editor",
diff --git a/Engine/Plugins/Experimental/Animation/RelativeIKOp/Source/BodyIntersectIKOp/BodyIntersectIKOp.Build.cs b/Engine/Plugins/Experimental/Animation/RelativeIKOp/Source/BodyIntersectIKOp/BodyIntersectIKOp.Build.cs
new file mode 100644
--- /dev/null
+++ b/Engine/Plugins/Experimental/Animation/RelativeIKOp/Source/BodyIntersectIKOp/BodyIntersectIKOp.Build.cs
@@ -0,0 +1,61 @@
+// Copyright Epic Games, Inc. All Rights Reserved.
+
+using UnrealBuildTool;
+using System.IO; // for Path
+
+namespace UnrealBuildTool.Rules
+{
+	public class BodyIntersectIKOp : ModuleRules
+	{
+		public BodyIntersectIKOp(ReadOnlyTargetRules Target) : base(Target)
+		{
+			PCHUsage = ModuleRules.PCHUsageMode.UseExplicitOrSharedPCHs;
+			CppCompileWarningSettings.UnsafeTypeCastWarningLevel = WarningLevel.Error;
+
+			PublicIncludePaths.AddRange(
+				new string[] {
+					// ... add public include paths required here ...
+				}
+				);
+
+
+			PrivateIncludePaths.AddRange(
+				new string[] {
+					// ... add other private include paths required here ...
+				}
+				);
+
+
+			PublicDependencyModuleNames.AddRange(
+				new string[]
+				{
+				}
+				);
+
+
+			PrivateDependencyModuleNames.AddRange(
+				new string[]
+				{
+				"Core",
+				"CoreUObject",
+				"Engine",
+				"AnimationCore",
+				"MeshDescription",
+				"StaticMeshDescription",
+				"SkeletalMeshDescription",
+				"Renderer",
+				"IKRig"
+				}
+				);
+
+
+			DynamicallyLoadedModuleNames.AddRange(
+				new string[]
+				{
+					// ... add any modules that your module loads dynamically here ...
+				}
+				);
+		}
+	}
+
+}
\ No newline at end of file
diff --git a/Engine/Plugins/Experimental/Animation/RelativeIKOp/Source/BodyIntersectIKOp/Private/BodyIntersectIKOp.cpp b/Engine/Plugins/Experimental/Animation/RelativeIKOp/Source/BodyIntersectIKOp/Private/BodyIntersectIKOp.cpp
new file mode 100644
--- /dev/null
+++ b/Engine/Plugins/Experimental/Animation/RelativeIKOp/Source/BodyIntersectIKOp/Private/BodyIntersectIKOp.cpp
@@ -0,0 +1,502 @@
+// Copyright Epic Games, Inc. All Rights Reserved.
+
+#include "BodyIntersectIKOp.h"
+
+#include "IKRigObjectVersion.h"
+#include "PhysicsEngine/PhysicsAsset.h"
+#include "PhysicsEngine/SkeletalBodySetup.h"
+#include "Retargeter/IKRetargetProcessor.h"
+#include "Retargeter/RetargetOps/RunIKRigOp.h"
+#include "Rig/Solvers/PointsToRotation.h"
+
+#if WITH_EDITOR
+#include "PrimitiveDrawingUtils.h"
+#include "Engine/EngineTypes.h"
+#endif
+
+#define LOCTEXT_NAMESPACE "BodyIntersectIKOp"
+
+
+const UClass* FIKRetargetBodyIntersectIKOpSettings::GetControllerType() const
+{
+	return UIKRetargetBodyIntersectController::StaticClass();
+}
+
+void FIKRetargetBodyIntersectIKOpSettings::CopySettingsAtRuntime(const FIKRetargetOpSettingsBase* InSettingsToCopyFrom)
+{
+	static TArray<FName> PropertiesToIgnore = {"PhysicsAssetOverride", "IntersectGoals", "IntersectBodies"};
+	FIKRetargetOpBase::CopyStructProperties(
+		FIKRetargetBodyIntersectIKOpSettings::StaticStruct(),
+		InSettingsToCopyFrom,
+		this,
+		PropertiesToIgnore);
+}
+
+bool FIKRetargetBodyIntersectIKOp::Initialize(
+	const FIKRetargetProcessor& InProcessor,
+	const FRetargetSkeleton& InSourceSkeleton,
+	const FTargetSkeleton& InTargetSkeleton,
+	const FIKRetargetOpBase* InParentOp,
+	FIKRigLogger& InLog)
+{
+	bIsInitialized = false;
+
+	// this op requires a parent to supply an IK Rig
+	if (!ensure(InParentOp))
+	{
+		return false;
+	}
+
+	// validate that an IK rig has been assigned
+	const FIKRetargetRunIKRigOp* ParentRigOp = reinterpret_cast<const FIKRetargetRunIKRigOp*>(InParentOp);
+	if (!ParentRigOp || ParentRigOp->Settings.IKRigAsset == nullptr)
+	{
+		InLog.LogWarning( FText::Format(
+		LOCTEXT("MissingIKRig", "{0}, is missing an IK rig. No chains can be retargeted. "), FText::FromName(GetName())));
+		return false;
+	}
+
+	PhysicsAsset = Settings.TargetPhysicsAssetOverride;
+
+	if (!PhysicsAsset)
+	{
+		InLog.LogWarning( FText::Format(
+		LOCTEXT("MissingPhysicsAsset", "{0}: Intersect Physics Asset must be specified. "), FText::FromName(GetName())));
+		return false;
+	}
+
+	UpdateCacheSkelInfo(InTargetSkeleton);
+	
+	const FIKRigGoalContainer& GoalContainer = InProcessor.GetIKRigGoalContainer();
+	for (const FName& GoalName : Settings.IntersectGoals)
+	{
+		const FIKRigGoal* IKGoal = GoalContainer.FindGoalByName(GoalName);
+		if (!IKGoal)
+		{
+			continue;
+		}
+
+		// TODO: Assumes goal bone has a phys body associated with it
+		UpdateTargetBoneMap(IKGoal->BoneName);
+	}
+	
+	bIsInitialized = true;
+	return bIsInitialized;
+}
+
+void FIKRetargetBodyIntersectIKOp::Run(
+	FIKRetargetProcessor& InProcessor,
+	const double InDeltaTime,
+	const TArray<FTransform>& InSourceGlobalPose,
+	TArray<FTransform>& OutTargetGlobalPose)
+{
+#if WITH_EDITOR
+	ResetDebugInfo();
+#endif //WITH_EDITOR
+	
+	if (InProcessor.IsIKForcedOff())
+	{
+		return;
+	}
+
+	if (!bIsInitialized)
+	{
+		return;
+	}
+
+#if WITH_EDITOR
+	FDebugBodyIntersectDrawInfo LocalDebugInfo;
+#endif //WITH_EDITOR
+
+	FIKRigGoalContainer& GoalContainer = InProcessor.GetIKRigGoalContainer();
+	for (const FName& GoalName : Settings.IntersectGoals)
+	{
+		FIKRigGoal* IKGoal = GoalContainer.FindGoalByName(GoalName);
+		if (!IKGoal || !IKGoal->bEnabled)
+		{
+			continue;
+		}
+
+		const FName TargetBoneName = IKGoal->BoneName;
+		if (!CacheTargetSkelIndices.Contains(TargetBoneName))
+		{
+			// UE_LOG(LogAnimation, Warning, TEXT("---- Could not find cache bone: %s"), *TargetBoneName.ToString());
+			continue;
+		}
+		
+		FKShapeElem* GoalShape = FindBodyShape(PhysicsAsset, IKGoal->BoneName);
+		if (!GoalShape)
+		{
+			continue;
+		}
+
+		const int32 TargetBoneIdx = CacheTargetSkelIndices[TargetBoneName];
+		const FTransform TargetBoneTfm = OutTargetGlobalPose[TargetBoneIdx];
+		
+		double GoalRadius = CalcShapeSmallRadius(GoalShape);
+		// TODO: Handle actual capsule offsets w/ a lever arm for rotation?
+		FVector CompGoalUpdate = GetGoalLocToCompSpace(IKGoal, TargetBoneTfm);
+		if (FMath::IsNearlyZero(GoalRadius))
+		{
+			continue;
+		}
+			
+		for (FName BodyName : Settings.IntersectBodies)
+		{
+			if (TargetBoneName == BodyName)
+			{
+				continue;
+			}
+				
+			if (!CacheTargetSkelIndices.Contains(BodyName))
+			{
+				continue;
+			}
+				
+			int32 BodyPoseIdx = CacheTargetSkelIndices[BodyName];
+			FTransform BodyTfm = OutTargetGlobalPose[BodyPoseIdx];
+			FKShapeElem* ShapeElem = FindBodyShape(PhysicsAsset, BodyName);
+			if (!ShapeElem)
+			{
+				continue;
+			}
+
+			FVector DeltaDir = FVector::ZeroVector;
+			double Dist = CalcIntersectionDelta(CompGoalUpdate, GoalRadius, BodyTfm, ShapeElem, DeltaDir);
+			if (Dist > 0.0)
+			{
+				CompGoalUpdate += Dist * DeltaDir;
+			}
+		}
+		SetGoalPosFromCompSpace(IKGoal, TargetBoneTfm, CompGoalUpdate);
+#if WITH_EDITOR
+		if (Settings.bEnableDebugDraw)
+		{
+			LocalDebugInfo.TestSpheres.Emplace(CompGoalUpdate, GoalRadius);
+		}
+#endif //WITH_EDITOR
+	}
+
+#if WITH_EDITOR
+	if (Settings.bEnableDebugDraw)
+	{
+		FScopeLock ScopeLock(&DebugDataMutex);
+		
+		DebugDrawInfo.TestSpheres = LocalDebugInfo.TestSpheres;
+
+		// Cache debug info for target domain bodies
+		DebugDrawInfo.TargetIntersectTfms.Reset(Settings.IntersectBodies.Num());
+		for (FName TargetBoneName : Settings.IntersectBodies)
+		{
+			// TODO: Only show domain bodies? 
+			if (!CacheTargetSkelIndices.Contains(TargetBoneName))
+			{
+				continue;
+			}
+			
+			int32 TargetPoseIdx = CacheTargetSkelIndices[TargetBoneName];
+			DebugDrawInfo.TargetIntersectTfms.Add({TargetBoneName,OutTargetGlobalPose[TargetPoseIdx]});
+		}
+	}
+#endif
+}
+
+FVector FIKRetargetBodyIntersectIKOp::GetGoalLocToCompSpace(const FIKRigGoal* Goal, const FTransform& BoneTfm) const
+{
+	switch (Goal->PositionSpace)
+	{
+	case EIKRigGoalSpace::Additive:
+		{
+			const FVector GoalPos = FMath::Lerp(FVector::ZeroVector, Goal->Position, Goal->PositionAlpha);
+			return BoneTfm.TransformPositionNoScale(GoalPos);
+		}
+	case EIKRigGoalSpace::Component:
+		{
+			const FVector GoalPos = FMath::Lerp(BoneTfm.GetLocation(), Goal->Position, Goal->PositionAlpha);
+			return GoalPos;
+		}
+	case EIKRigGoalSpace::World:
+	default:
+		{
+			// TODO: Handle world-space case
+			return Goal->Position;
+		}
+	}
+}
+
+void FIKRetargetBodyIntersectIKOp::SetGoalPosFromCompSpace(FIKRigGoal* Goal, const FTransform& BoneTfm, const FVector& CompSpaceLoc) const
+{
+	// TODO: Alpha passthrough and/or alpha ignore
+	switch (Goal->PositionSpace)
+	{
+	case EIKRigGoalSpace::Additive:
+		{
+			Goal->Position = BoneTfm.InverseTransformPositionNoScale(CompSpaceLoc);
+			Goal->PositionAlpha = 1.0f;
+			break;
+		}
+	case EIKRigGoalSpace::Component:
+		{
+			Goal->Position = CompSpaceLoc;
+			Goal->PositionAlpha = 1.0f;
+			break;
+		}
+	case EIKRigGoalSpace::World:
+	default:
+		{
+			// TODO: Handle world-space case
+			Goal->Position = CompSpaceLoc;
+			Goal->PositionAlpha = 1.0f;
+			break;
+		}
+	}
+}
+
+double FIKRetargetBodyIntersectIKOp::CalcIntersectionDelta(const FVector& TargetPoint, double TargetRadius, const FTransform& ShapeTfm, FKShapeElem* ShapeElem, FVector& OutDeltaDir)
+{
+	if (!ShapeElem)
+	{
+		return -1.0;
+	}
+
+	float Dist = 0.0f;
+	FVector ClosestPoint;
+	switch (ShapeElem->GetShapeType())
+	{
+	case EAggCollisionShape::Sphere:
+		{
+			FKSphereElem* SphereElem = static_cast<FKSphereElem*>(ShapeElem);
+			Dist = SphereElem->GetClosestPointAndNormal(TargetPoint, ShapeTfm, ClosestPoint, OutDeltaDir);
+			break;
+		}
+
+	case EAggCollisionShape::Sphyl:
+		{
+			FKSphylElem* CapsuleElem = static_cast<FKSphylElem*>(ShapeElem);
+			Dist = CapsuleElem->GetClosestPointAndNormal(TargetPoint, ShapeTfm, ClosestPoint, OutDeltaDir);
+			break;
+		}
+
+	// case EAggCollisionShape::Box:
+	// 	{
+	// 		const FKBoxElem* BoxElem = static_cast<const FKBoxElem*>(ShapeElem);
+	// 		Dist = BoxElem->GetClosestPointAndNormal(TargetPoint, ShapeTfm, ClosestPoint, OutDeltaDir);
+	// 		break;
+	// 	}
+
+	default:
+		{
+			return -1.0;
+		}
+	}
+
+	return TargetRadius - Dist;
+}
+
+double FIKRetargetBodyIntersectIKOp::CalcShapeSmallRadius(FKShapeElem* ShapeElem)
+{
+	switch (ShapeElem->GetShapeType())
+	{
+	case EAggCollisionShape::Sphere:
+		{
+			FKSphereElem* SphereElem = static_cast<FKSphereElem*>(ShapeElem);
+			return SphereElem->Radius;
+		}
+
+	case EAggCollisionShape::Sphyl:
+		{
+			FKSphylElem* CapsuleElem = static_cast<FKSphylElem*>(ShapeElem);
+			return CapsuleElem->Radius;
+		}
+
+	default:
+		{
+			return 0.0;
+		}
+	}
+}
+
+double FIKRetargetBodyIntersectIKOp::CalcShapeAvgRadius(FKShapeElem* ShapeElem)
+{
+	switch (ShapeElem->GetShapeType())
+	{
+	case EAggCollisionShape::Sphere:
+		{
+			FKSphereElem* SphereElem = static_cast<FKSphereElem*>(ShapeElem);
+			return SphereElem->Radius;
+		}
+
+	case EAggCollisionShape::Sphyl:
+		{
+			FKSphylElem* CapsuleElem = static_cast<FKSphylElem*>(ShapeElem);
+			return CapsuleElem->Radius + 0.25*CapsuleElem->Length;
+		}
+
+	default:
+		{
+			return 0.0;
+		}
+	}
+}
+
+FKShapeElem* FIKRetargetBodyIntersectIKOp::FindBodyShape(const UPhysicsAsset* PhysAsset, FName BoneName)
+{
+	int32 BodyIdx = PhysAsset->FindBodyIndex(BoneName);
+	if (BodyIdx == INDEX_NONE)
+	{
+		UE_LOG(LogAnimation, Warning, TEXT("No body index found: %s"), *BoneName.ToString());
+		return nullptr;
+	}
+	return PhysAsset->SkeletalBodySetups[BodyIdx]->AggGeom.GetElement(0);
+}
+
+void FIKRetargetBodyIntersectIKOp::UpdateTargetBoneMap(FName TargetBoneName)
+{
+	int32 TargetIdx = TargetBoneNames.Find(TargetBoneName);
+
+	if (TargetIdx == INDEX_NONE)
+	{
+		return;
+	}
+	
+	CacheTargetSkelIndices.Emplace(TargetBoneName, TargetIdx);
+}
+
+void FIKRetargetBodyIntersectIKOp::UpdateCacheSkelInfo(const FTargetSkeleton& InTargetSkeleton)
+{
+	TargetBoneNames = InTargetSkeleton.BoneNames;
+	CacheTargetSkelIndices.Reset();
+	
+	for (FName TargetBodyName : Settings.IntersectBodies)
+	{
+		UpdateTargetBoneMap(TargetBodyName);
+	}
+}
+
+void FIKRetargetBodyIntersectIKOp::OnAddedToStack(const UIKRetargeter* InRetargetAsset, const FIKRetargetOpBase* InParentOp)
+{
+}
+
+FIKRetargetOpSettingsBase* FIKRetargetBodyIntersectIKOp::GetSettings()
+{
+	return &Settings;
+}
+
+const UScriptStruct* FIKRetargetBodyIntersectIKOp::GetSettingsType() const
+{
+	return FIKRetargetBodyIntersectIKOpSettings::StaticStruct();
+}
+
+const UScriptStruct* FIKRetargetBodyIntersectIKOp::GetType() const
+{
+	return FIKRetargetBodyIntersectIKOp::StaticStruct();
+}
+
+const UScriptStruct* FIKRetargetBodyIntersectIKOp::GetParentOpType() const
+{
+	return FIKRetargetRunIKRigOp::StaticStruct();
+}
+
+#if WITH_EDITOR
+
+FCriticalSection FIKRetargetBodyIntersectIKOp::DebugDataMutex;
+
+void FIKRetargetBodyIntersectIKOp::DebugDraw(
+	FPrimitiveDrawInterface* InPDI,
+	const FTransform& InSourceTransform,
+	const FTransform& InComponentTransform,
+	const double InComponentScale,
+	const FIKRetargetDebugDrawState& InEditorState) const
+{
+	if (!Settings.bEnableDebugDraw)
+	{
+		return;
+	}
+	
+	FScopeLock ScopeLock(&DebugDataMutex);
+
+	if (!Settings.TargetPhysicsAssetOverride)
+	{
+		return;
+	}
+
+	for (const auto& IntersectBody : DebugDrawInfo.TargetIntersectTfms)
+	{
+		const FName BodyName = IntersectBody.Get<0>();
+		const FKShapeElem* ShapeElem = FindBodyShape(Settings.TargetPhysicsAssetOverride, BodyName);
+		if (!ShapeElem)
+		{
+			continue;
+		}
+		
+		FTransform CompTfm = IntersectBody.Get<1>() * InComponentTransform;
+		DebugDrawPhysBody(InPDI, CompTfm, InComponentScale, ShapeElem, FLinearColor::Yellow);
+	}
+
+	for (const auto& Sphere : DebugDrawInfo.TestSpheres)
+	{
+		FVector Center = InComponentTransform.TransformPosition(Sphere.Get<0>());
+		DrawWireSphere(InPDI, Center, FLinearColor::Red, Sphere.Get<1>()*InComponentScale, 25, SDPG_Foreground);
+	}
+}
+
+void FIKRetargetBodyIntersectIKOp::DebugDrawPhysBody(FPrimitiveDrawInterface* InPDI, const FTransform& ParentTransform, double Scale, const FKShapeElem* ShapeElem, const FLinearColor& Color) const
+{
+	FTransform BodyFrame = ShapeElem->GetTransform() * ParentTransform;
+
+	const FVector Translation = BodyFrame.GetLocation();
+	const FVector UnitXAxis = BodyFrame.GetUnitAxis( EAxis::X );
+	const FVector UnitYAxis = BodyFrame.GetUnitAxis( EAxis::Y );
+	const FVector UnitZAxis = BodyFrame.GetUnitAxis( EAxis::Z );
+	
+	switch (ShapeElem->GetShapeType())
+	{
+	case EAggCollisionShape::Box:
+		{
+			const FKBoxElem* BoxElem = static_cast<const FKBoxElem*>(ShapeElem);
+			const FVector Extent = 0.5 * FVector(BoxElem->X, BoxElem->Y, BoxElem->Z) * Scale;
+			DrawOrientedWireBox(InPDI, Translation, UnitXAxis, UnitYAxis, UnitZAxis, Extent, Color, SDPG_Foreground);
+			return;
+		}
+	case EAggCollisionShape::Sphyl:
+		{
+			const FKSphylElem* CapsuleElem = static_cast<const FKSphylElem*>(ShapeElem);
+			const double Radius = CapsuleElem->Radius * Scale;
+			const double HalfHeight = (0.5*CapsuleElem->Length + CapsuleElem->Radius) * Scale;
+			DrawWireCapsule(InPDI, Translation, UnitXAxis, UnitYAxis, UnitZAxis, Color,
+						Radius, HalfHeight,25, SDPG_Foreground, 0, 1.0);
+			return;
+		}
+	case EAggCollisionShape::Sphere:
+		{
+			const FKSphereElem* SphereElem = static_cast<const FKSphereElem*>(ShapeElem);
+			const double Radius = SphereElem->Radius * Scale;
+			DrawWireSphere(InPDI, Translation, Color, Radius, 25, SDPG_Foreground);
+			return;
+		}
+	default:
+		return;
+	}
+}
+
+
+
+void FIKRetargetBodyIntersectIKOp::ResetDebugInfo()
+{
+	DebugDrawInfo.TargetIntersectTfms.Reset();
+	DebugDrawInfo.TestSpheres.Reset();
+}
+
+#endif //WITH_EDITOR
+
+FIKRetargetBodyIntersectIKOpSettings UIKRetargetBodyIntersectController::GetSettings()
+{
+	return *reinterpret_cast<FIKRetargetBodyIntersectIKOpSettings*>(OpSettingsToControl);
+}
+
+void UIKRetargetBodyIntersectController::SetSettings(FIKRetargetBodyIntersectIKOpSettings InSettings)
+{
+	OpSettingsToControl->CopySettingsAtRuntime(&InSettings);
+}
+
+#undef LOCTEXT_NAMESPACE
\ No newline at end of file
diff --git a/Engine/Plugins/Experimental/Animation/RelativeIKOp/Source/BodyIntersectIKOp/Private/BodyIntersectIKOpModule.cpp b/Engine/Plugins/Experimental/Animation/RelativeIKOp/Source/BodyIntersectIKOp/Private/BodyIntersectIKOpModule.cpp
new file mode 100644
--- /dev/null
+++ b/Engine/Plugins/Experimental/Animation/RelativeIKOp/Source/BodyIntersectIKOp/Private/BodyIntersectIKOpModule.cpp
@@ -0,0 +1,17 @@
+// Copyright Epic Games, Inc. All Rights Reserved.
+
+#include "BodyIntersectIKOpModule.h"
+
+#define LOCTEXT_NAMESPACE "BodyIntersectIKOpModule"
+
+void FBodyIntersectIKOpModule::StartupModule()
+{
+}
+
+void FBodyIntersectIKOpModule::ShutdownModule()
+{
+}
+
+#undef LOCTEXT_NAMESPACE
+    
+IMPLEMENT_MODULE(FBodyIntersectIKOpModule, BodyIntersectIKOp)
\ No newline at end of file
diff --git a/Engine/Plugins/Experimental/Animation/RelativeIKOp/Source/BodyIntersectIKOp/Public/BodyIntersectIKOp.h b/Engine/Plugins/Experimental/Animation/RelativeIKOp/Source/BodyIntersectIKOp/Public/BodyIntersectIKOp.h
new file mode 100644
--- /dev/null
+++ b/Engine/Plugins/Experimental/Animation/RelativeIKOp/Source/BodyIntersectIKOp/Public/BodyIntersectIKOp.h
@@ -0,0 +1,149 @@
+// Copyright Epic Games, Inc. All Rights Reserved.
+#pragma once
+
+#include "Retargeter/IKRetargetOps.h"
+#include "Retargeter/IKRetargetSettings.h"
+
+#include "BodyIntersectIKOp.generated.h"
+
+#define UE_API BODYINTERSECTIKOP_API
+
+#define LOCTEXT_NAMESPACE "BodyIntersectIKOp"
+
+struct FIKRigGoal;
+struct FKShapeElem;
+class UPhysicsAsset;
+
+// Used to hold all debug draw info for convenience
+struct FDebugBodyIntersectDrawInfo
+{
+	// Intersect info
+	TArray<TTuple<FName,FTransform>> TargetIntersectTfms;
+	TArray<TTuple<FVector,double>> TestSpheres;
+};
+
+USTRUCT(BlueprintType, meta = (DisplayName = "Body Intersect IK Settings"))
+struct FIKRetargetBodyIntersectIKOpSettings : public FIKRetargetOpSettingsBase
+{
+	GENERATED_BODY()
+
+	UE_API virtual const UClass* GetControllerType() const override;
+
+	UE_API virtual void CopySettingsAtRuntime(const FIKRetargetOpSettingsBase* InSettingsToCopyFrom) override;
+
+	// Target physics asset for checking intersections against
+	UPROPERTY(EditAnywhere, BlueprintReadWrite, Category = "Physics", meta=(ReinitializeOnEdit))
+	TObjectPtr<UPhysicsAsset> TargetPhysicsAssetOverride;
+
+	// IK Goals to check sphere intersections (Goal effector bone should have a sphere or capsule body!)
+	UPROPERTY(EditAnywhere, BlueprintReadWrite, Category = "Intersection", meta=(ReinitializeOnEdit))
+	TArray<FName> IntersectGoals;
+
+	// Physics bodies to do trivial intersection against
+	UPROPERTY(EditAnywhere, BlueprintReadWrite, Category = "Intersection", meta=(ReinitializeOnEdit))
+	TArray<FName> IntersectBodies;
+	
+	// Display source and target physics bodies for baked data
+	UPROPERTY(EditAnywhere, BlueprintReadWrite, Category = "Debug")
+	bool bEnableDebugDraw = false;
+};
+
+USTRUCT(BlueprintType, meta = (DisplayName = "Body Intersect Goals"))
+struct FIKRetargetBodyIntersectIKOp : public FIKRetargetOpBase
+{
+	GENERATED_BODY()
+	
+	UE_API virtual bool Initialize(
+		const FIKRetargetProcessor& InProcessor,
+		const FRetargetSkeleton& InSourceSkeleton,
+		const FTargetSkeleton& InTargetSkeleton,
+		const FIKRetargetOpBase* InParentOp,
+		FIKRigLogger& InLog) override;
+	
+	UE_API virtual void Run(
+		FIKRetargetProcessor& InProcessor,
+		const double InDeltaTime,
+		const TArray<FTransform>& InSourceGlobalPose,
+		TArray<FTransform>& OutTargetGlobalPose) override;
+
+	UE_API virtual void OnAddedToStack(const UIKRetargeter* InRetargetAsset, const FIKRetargetOpBase* InParentOp) override;
+
+	UE_API virtual FIKRetargetOpSettingsBase* GetSettings() override;
+	
+	UE_API virtual const UScriptStruct* GetSettingsType() const override;
+	
+	UE_API virtual const UScriptStruct* GetType() const override;
+	
+	UE_API virtual const UScriptStruct* GetParentOpType() const override;
+
+	UPROPERTY()
+	FIKRetargetBodyIntersectIKOpSettings Settings;
+
+private:
+	void UpdateCacheSkelInfo(const FTargetSkeleton& InTargetSkeleton);
+	void UpdateTargetBoneMap(FName TargetBoneName);
+
+	FVector GetGoalLocToCompSpace(const FIKRigGoal* Goal, const FTransform& BoneTfm) const;
+	void SetGoalPosFromCompSpace(FIKRigGoal* Goal, const FTransform& BoneTfm, const FVector& CompSpaceLoc) const;
+	
+	static double CalcIntersectionDelta(const FVector&  TargetPoint, double TargetRadius, const FTransform& ShapeTfm, FKShapeElem* ShapeElem, FVector& OutDeltaDir);
+	
+	static double CalcShapeSmallRadius(FKShapeElem* ShapeElem);
+	static double CalcShapeAvgRadius(FKShapeElem* ShapeElem);
+
+	static FKShapeElem* FindBodyShape(const UPhysicsAsset* PhysAsset, FName BoneName);
+
+#if WITH_EDITOR
+public:
+	UE_API virtual void DebugDraw(
+		FPrimitiveDrawInterface* InPDI,
+		const FTransform& InSourceTransform,
+		const FTransform& InComponentTransform,
+		const double InComponentScale,
+		const FIKRetargetDebugDrawState& InEditorState) const override;
+
+private:
+	void DebugDrawPhysBody(FPrimitiveDrawInterface* InPDI, const FTransform& ParentTransform, double Scale, const FKShapeElem* ShapeElem, const FLinearColor& Color) const;
+	
+	void ResetDebugInfo();
+
+	// Collection of all debug structures
+	FDebugBodyIntersectDrawInfo DebugDrawInfo;
+	
+	static UE_API FCriticalSection DebugDataMutex;
+#endif
+private:
+	TObjectPtr<const UPhysicsAsset> PhysicsAsset;
+
+	// Cache ik goals to check
+	TArray<const FIKRigGoal*> IntersectIKGoals;
+	
+	// Cache source/target bone-names
+	TArray<FName> TargetBoneNames;
+
+	// Cache target skel index info
+	TMap<FName,int32> CacheTargetSkelIndices;
+};
+
+/* The blueprint/python API for editing a Stride Warping Op */
+UCLASS(MinimalAPI, BlueprintType)
+class UIKRetargetBodyIntersectController : public UIKRetargetOpControllerBase
+{
+	GENERATED_BODY()
+	
+public:
+
+	/* Get the current op settings as a struct.
+	 * @return FIKRetargetRelativeIKOpSettings struct with the current settings used by the op. */
+	UFUNCTION(BlueprintCallable, Category = Settings)
+	UE_API FIKRetargetBodyIntersectIKOpSettings GetSettings();
+
+	/* Set the solver settings. Input is a custom struct type for this solver.
+	 * @param InSettings a FIKRetargetRelativeIKOpSettings struct containing all the settings to apply to this op */
+	UFUNCTION(BlueprintCallable, Category = Settings)
+	UE_API void SetSettings(FIKRetargetBodyIntersectIKOpSettings InSettings);
+};
+
+#undef LOCTEXT_NAMESPACE
+
+#undef UE_API
diff --git a/Engine/Plugins/Experimental/Animation/RelativeIKOp/Source/BodyIntersectIKOp/Public/BodyIntersectIKOpModule.h b/Engine/Plugins/Experimental/Animation/RelativeIKOp/Source/BodyIntersectIKOp/Public/BodyIntersectIKOpModule.h
new file mode 100644
--- /dev/null
+++ b/Engine/Plugins/Experimental/Animation/RelativeIKOp/Source/BodyIntersectIKOp/Public/BodyIntersectIKOpModule.h
@@ -0,0 +1,14 @@
+// Copyright Epic Games, Inc. All Rights Reserved.
+
+#pragma once
+
+#include "Modules/ModuleManager.h"
+
+class FBodyIntersectIKOpModule : public IModuleInterface
+{
+public:
+
+	/** IModuleInterface implementation */
+	virtual void StartupModule() override;
+	virtual void ShutdownModule() override;
+};
diff --git a/Engine/Plugins/Experimental/Animation/RelativeIKOp/Source/RelativeIKOp/Private/RelativeIKOp.cpp b/Engine/Plugins/Experimental/Animation/RelativeIKOp/Source/RelativeIKOp/Private/RelativeIKOp.cpp
--- a/Engine/Plugins/Experimental/Animation/RelativeIKOp/Source/RelativeIKOp/Private/RelativeIKOp.cpp
+++ b/Engine/Plugins/Experimental/Animation/RelativeIKOp/Source/RelativeIKOp/Private/RelativeIKOp.cpp
@@ -29,7 +29,7 @@ const UClass* FIKRetargetRelativeIKOpSettings::GetControllerType() const
 
 void FIKRetargetRelativeIKOpSettings::CopySettingsAtRuntime(const FIKRetargetOpSettingsBase* InSettingsToCopyFrom)
 {
-	static TArray<FName> PropertiesToIgnore = {"SourcePhysicsAsset", "TargetPhysicsAsset", "DomainBodies", "ContactBodies", "IntersectBodies"};
+	static TArray<FName> PropertiesToIgnore = {"SourcePhysicsAssetOverride", "TargetPhysicsAssetOverride", "DomainBodies", "ContactBodies"};
 	FIKRetargetOpBase::CopyStructProperties(
 		FIKRetargetRelativeIKOpSettings::StaticStruct(),
 		InSettingsToCopyFrom,
@@ -57,7 +57,7 @@ bool FIKRetargetRelativeIKOp::Initialize(
 	if (!ParentRigOp || ParentRigOp->Settings.IKRigAsset == nullptr)
 	{
 		InLog.LogWarning( FText::Format(
-		LOCTEXT("MissingIKRig", "{0}, is missing an IK rig. No chains can be retargeted. "), FText::FromName(GetName())));
+		LOCTEXT("MissingIKRig", "{0}, is missing an IK rig. No goals can be updated. "), FText::FromName(GetName())));
 		return false;
 	}
 
@@ -359,59 +359,6 @@ void FIKRetargetRelativeIKOp::Run(
 				IKGoal->PositionAlpha = 1.0;
 				IKGoal->PositionSpace = EIKRigGoalSpace::Component;
 			}
-
-			// Sphere-capsule intersection goal deltas
-			if (Settings.bTestIntersectionFixup)
-			{
-				FKShapeElem* BodyShape = FindBodyShape(TargetPhysicsAsset, TargetBoneName);
-				double BodyRadius = CalcShapeSmallRadius(BodyShape);
-				FVector LocalOffset = BodyShape->GetTransform().GetLocation();
-				FVector BodyOffset = OutTargetGlobalPose[TargetBoneIdx].TransformVector(LocalOffset);
-				FVector BodyLocation = IKGoal->Position + BodyOffset;
-				if (FMath::IsNearlyZero(BodyRadius))
-				{
-					continue;
-				}
-				
-				for (FName BodyName : Settings.IntersectBodies)
-				{
-					if (TargetBoneName == BodyName)
-					{
-						continue;
-					}
-					
-					if (!CacheTargetSkelIndices.Contains(BodyName))
-					{
-						continue;
-					}
-					
-					int32 BodyPoseIdx = CacheTargetSkelIndices[BodyName];
-					FTransform BodyTfm = OutTargetGlobalPose[BodyPoseIdx];
-					FKShapeElem* ShapeElem = FindBodyShape(TargetPhysicsAsset, BodyName);
-					if (!ShapeElem)
-					{
-						continue;
-					}
-
-					FVector DeltaDir = FVector::ZeroVector;
-					double Dist = CalcIntersectionDelta(BodyLocation, BodyRadius, BodyTfm, ShapeElem, DeltaDir);
-					if (Dist > 0.0)
-					{
-						BodyLocation += Dist * DeltaDir;
-					}
-				}
-				IKGoal->Position = BodyLocation - BodyOffset;
-			}
-#if WITH_EDITOR
-			if (Settings.bEnableDebugDraw)
-			{
-				FKShapeElem* BodyShape = FindBodyShape(TargetPhysicsAsset, TargetBoneName);
-				double BodyRadius = CalcShapeSmallRadius(BodyShape);
-				FVector LocalOffset = BodyShape->GetTransform().GetLocation();
-				FVector BodyOffset = OutTargetGlobalPose[TargetBoneIdx].TransformVector(LocalOffset);
-				LocalDebugInfo.TestSpheres.Emplace(IKGoal->Position+BodyOffset, BodyRadius);
-			}
-#endif //WITH_EDITOR
 		}
 	}
 
@@ -424,7 +371,6 @@ void FIKRetargetRelativeIKOp::Run(
 		DebugDrawInfo.PairRetargetInfo = LocalDebugInfo.PairRetargetInfo;
 		DebugDrawInfo.SourcePairVerts = LocalDebugInfo.SourcePairVerts;
 		DebugDrawInfo.TargetPairVerts = LocalDebugInfo.TargetPairVerts;
-		DebugDrawInfo.TestSpheres = LocalDebugInfo.TestSpheres;
 
 		// Cache debug info for source bodies
 		DebugDrawInfo.SourceBodyInfo.Reset(CacheSourceEffectBones.Num());
@@ -550,93 +496,6 @@ double FIKRetargetRelativeIKOp::CalcLineSphereIntersect(const FVector& Center, d
 	return (-b + FMath::Sqrt(b*b - a*c)) / a;
 }
 
-double FIKRetargetRelativeIKOp::CalcIntersectionDelta(const FVector& TargetPoint, double TargetRadius, const FTransform& ShapeTfm, FKShapeElem* ShapeElem, FVector& OutDeltaDir)
-{
-	if (!ShapeElem)
-	{
-		return -1.0;
-	}
-
-	float Dist = 0.0f;
-	FVector ClosestPoint;
-	switch (ShapeElem->GetShapeType())
-	{
-	case EAggCollisionShape::Sphere:
-		{
-			FKSphereElem* SphereElem = static_cast<FKSphereElem*>(ShapeElem);
-			Dist = SphereElem->GetClosestPointAndNormal(TargetPoint, ShapeTfm, ClosestPoint, OutDeltaDir);
-			break;
-		}
-
-	case EAggCollisionShape::Sphyl:
-		{
-			FKSphylElem* CapsuleElem = static_cast<FKSphylElem*>(ShapeElem);
-			Dist = CapsuleElem->GetClosestPointAndNormal(TargetPoint, ShapeTfm, ClosestPoint, OutDeltaDir);
-			break;
-		}
-
-	// case EAggCollisionShape::Box:
-	// 	{
-	// 		const FKBoxElem* BoxElem = static_cast<const FKBoxElem*>(ShapeElem);
-	// 		Dist = BoxElem->GetClosestPointAndNormal(TargetPoint, ShapeTfm, ClosestPoint, OutDeltaDir);
-	// 		break;
-	// 	}
-
-	default:
-		{
-			return -1.0;
-		}
-	}
-
-	return TargetRadius - Dist;
-}
-
-double FIKRetargetRelativeIKOp::CalcShapeSmallRadius(FKShapeElem* ShapeElem)
-{
-	switch (ShapeElem->GetShapeType())
-	{
-	case EAggCollisionShape::Sphere:
-		{
-			FKSphereElem* SphereElem = static_cast<FKSphereElem*>(ShapeElem);
-			return SphereElem->Radius;
-		}
-
-	case EAggCollisionShape::Sphyl:
-		{
-			FKSphylElem* CapsuleElem = static_cast<FKSphylElem*>(ShapeElem);
-			return CapsuleElem->Radius;
-		}
-
-	default:
-		{
-			return 0.0;
-		}
-	}
-}
-
-double FIKRetargetRelativeIKOp::CalcShapeAvgRadius(FKShapeElem* ShapeElem)
-{
-	switch (ShapeElem->GetShapeType())
-	{
-	case EAggCollisionShape::Sphere:
-		{
-			FKSphereElem* SphereElem = static_cast<FKSphereElem*>(ShapeElem);
-			return SphereElem->Radius;
-		}
-
-	case EAggCollisionShape::Sphyl:
-		{
-			FKSphylElem* CapsuleElem = static_cast<FKSphylElem*>(ShapeElem);
-			return CapsuleElem->Radius + 0.25*CapsuleElem->Length;
-		}
-
-	default:
-		{
-			return 0.0;
-		}
-	}
-}
-
 FTransform FIKRetargetRelativeIKOp::CalcReferenceShapeScale3D(FKShapeElem* ShapeElem)
 {
 	FTransform ElementScale3D = FTransform::Identity;
@@ -811,11 +670,6 @@ void FIKRetargetRelativeIKOp::UpdateCacheSkelInfo(const FRetargetSkeleton& InSou
 	{
 		UpdateAllBoneMaps(SourceBodyName);
 	}
-
-	for (FName TargetBodyName : Settings.IntersectBodies)
-	{
-		UpdateTargetBoneMap(TargetBodyName);
-	}
 }
 
 
@@ -991,7 +845,7 @@ void FIKRetargetRelativeIKOp::DebugDraw(
 		{
 			DebugDrawBodies(InPDI, SourceNoScaleTfm, SourceScale, Settings.SourcePhysicsAssetOverride, DebugDrawInfo.SourceBodyInfo);
 		}
-		if (Settings.SourcePhysicsAssetOverride)
+		if (Settings.TargetPhysicsAssetOverride)
 		{
 			DebugDrawBodies(InPDI, InComponentTransform, InComponentScale, Settings.TargetPhysicsAssetOverride, DebugDrawInfo.TargetBodyInfo);
 		}
@@ -1008,15 +862,6 @@ void FIKRetargetRelativeIKOp::DebugDraw(
 		// Show IK Goals (and contribs)
 		DebugDrawGoalContributions(InPDI, InComponentTransform, InComponentScale, DebugDrawInfo.TargetGoals);
 	}
-
-	if (Settings.bDebugDrawIntersectionSpheres)
-	{
-		for ( const auto& Sphere : DebugDrawInfo.TestSpheres )
-		{
-			FVector Center = InComponentTransform.TransformVector(Sphere.Get<0>());
-			DrawWireSphere(InPDI, Center, FLinearColor::Red, Sphere.Get<1>()*InComponentScale, 25, SDPG_Foreground);
-		}
-	}
 }
 
 void FIKRetargetRelativeIKOp::DebugDrawBodyPairs(FPrimitiveDrawInterface* InPDI, const FTransform& BaseTransform, double Scale, const TArray<FDebugDrawBodyPair>& BodyPairs) const
@@ -1184,8 +1029,6 @@ void FIKRetargetRelativeIKOp::ResetDebugInfo()
 	DebugDrawInfo.SourceBodyInfo.Reset();
 	DebugDrawInfo.TargetBodyInfo.Reset();
 	
-	DebugDrawInfo.TestSpheres.Reset();
-	
 	DebugDrawInfo.PairRetargetInfo.Reset();
 	DebugDrawInfo.TargetGoals.Reset();
 }
diff --git a/Engine/Plugins/Experimental/Animation/RelativeIKOp/Source/RelativeIKOp/Public/RelativeIKOp.h b/Engine/Plugins/Experimental/Animation/RelativeIKOp/Source/RelativeIKOp/Public/RelativeIKOp.h
--- a/Engine/Plugins/Experimental/Animation/RelativeIKOp/Source/RelativeIKOp/Public/RelativeIKOp.h
+++ b/Engine/Plugins/Experimental/Animation/RelativeIKOp/Source/RelativeIKOp/Public/RelativeIKOp.h
@@ -64,10 +64,6 @@ struct FDebugRelativeIKDrawInfo
 	TArray<FDebugDrawBodyInfo> SourceBodyInfo;
 	TArray<FDebugDrawBodyInfo> TargetBodyInfo;
 
-	// Intersect info
-	// TArray<FTransform> TargetIntersectTfms;
-	TArray<TTuple<FVector,double>> TestSpheres;
-
 	// Debug info for retarget pairs
 	TArray<FDebugRelativeTargetPairSpace> PairRetargetInfo;
 	// Final goal points with contributions
@@ -102,10 +98,6 @@ struct FIKRetargetRelativeIKOpSettings : public FIKRetargetOpSettingsBase
 	// Physics body names that should be moved (IK goals) to maintain relative spatial relationships
 	UPROPERTY(EditAnywhere, BlueprintReadWrite, Category = "Body Pairs", meta=(ReinitializeOnEdit))
 	TArray<FName> ContactBodies;
-
-	// Physics body names to do trivial intersection against Contact Bodies
-	UPROPERTY(EditAnywhere, BlueprintReadWrite, Category = "Body Pairs", meta=(ReinitializeOnEdit))
-	TArray<FName> IntersectBodies;
 	
 	// Maximum distance for which body pair info is baked
 	UPROPERTY(EditAnywhere, BlueprintReadWrite, Category = "Parameters", meta=(ClampMin = "0.1"))
@@ -159,10 +151,6 @@ struct FIKRetargetRelativeIKOpSettings : public FIKRetargetOpSettingsBase
 	UPROPERTY(EditAnywhere, BlueprintReadWrite, Category = "Debug")
 	bool bDebugDrawPhysicsBodies = false;
 
-	// Display source and target physics bodies for baked data
-	UPROPERTY(EditAnywhere, BlueprintReadWrite, Category = "Debug")
-	bool bDebugDrawIntersectionSpheres = false;
-
 	// Run op and display debug info but DON'T update IK Goals
 	UPROPERTY(EditAnywhere, BlueprintReadWrite, Category = "Debug")
 	bool bDryRun = false;
@@ -175,10 +163,6 @@ struct FIKRetargetRelativeIKOpSettings : public FIKRetargetOpSettingsBase
 	UPROPERTY(EditAnywhere, BlueprintReadWrite, Category = "Test Features")
 	bool bTestFeasibilityWeight = true;
 
-	// Move goals to avoid intersection with bodies in Intersect Bodies list
-	UPROPERTY(EditAnywhere, BlueprintReadWrite, Category = "Test Features")
-	bool bTestIntersectionFixup = true;
-
 	// Use rotation-invariant space for physics body components on source/target
 	UPROPERTY(EditAnywhere, BlueprintReadWrite, Category = "Test Features")
 	bool bTestBodyRotInvariance = true;
@@ -234,13 +218,8 @@ private:
 	FName ApplyBodyMap(FName BodyName);
 	void UpdateCacheChainInfo(const TArray<FTransform>& TargetPose, FName BoneName);
 	void ComputeFeasibilityRange(FDoubleInterval& OutFeasibleRange, FName BoneName, FVector BoneChainOffset, FVector TargetRelativeToBoneV, FVector TargetRelativeToOpposeBoneV);
-	// void FixupSphereIntersections();
 
 	static double CalcLineSphereIntersect(const FVector& Center, double Radius, const FVector& StartPoint, const FVector& EndPoint);
-	static double CalcIntersectionDelta(const FVector&  TargetPoint, double TargetRadius, const FTransform& ShapeTfm, FKShapeElem* ShapeElem, FVector& OutDeltaDir);
-	
-	static double CalcShapeSmallRadius(FKShapeElem* ShapeElem);
-	static double CalcShapeAvgRadius(FKShapeElem* ShapeElem);
 
 	static FTransform CalcReferenceShapeScale3D(FKShapeElem* ShapeElem);
 	static FTransform GetBodyRotation(const UPhysicsAsset* PhysAsset, FName BoneName);
