RelativeIKOp: Fixed goal blending (WIP) if relativeIK is running after other goal ops.

diff --git a/Engine/Plugins/Experimental/Animation/RelativeIKOp/Source/RelativeIKOp/Private/RelativeIKOp.cpp b/Engine/Plugins/Experimental/Animation/RelativeIKOp/Source/RelativeIKOp/Private/RelativeIKOp.cpp
--- a/Engine/Plugins/Experimental/Animation/RelativeIKOp/Source/RelativeIKOp/Private/RelativeIKOp.cpp
+++ b/Engine/Plugins/Experimental/Animation/RelativeIKOp/Source/RelativeIKOp/Private/RelativeIKOp.cpp
@@ -347,9 +347,9 @@ void FIKRetargetRelativeIKOp::Run(
 			
 			// TODO: Play with total weight functions to alpha goal
 			float GoalAlpha = FMath::Min(1.0f, static_cast<float>(TotalWeight*FMath::Clamp(Settings.ContributionSumWeight, 0.0f, 1.0f)));
-			// TODO: Blend from old goal value (or target bone)
 			int32 TargetBoneIdx = CacheTargetSkelIndices[TargetBoneName];
-			FVector BlendGoal = FMath::Lerp(OutTargetGlobalPose[TargetBoneIdx].GetLocation(), GoalPoint, GoalAlpha);
+			FVector PrevGoal = CalcCompSpaceBlendLocation(IKGoal, OutTargetGlobalPose[TargetBoneIdx]);
+			FVector BlendGoal = FMath::Lerp(PrevGoal, GoalPoint, GoalAlpha);
 
 			bool bSetGoals = !Settings.bEnableDebugDraw || !Settings.bDryRun;
 			if (bSetGoals)
@@ -565,6 +565,30 @@ FTransform FIKRetargetRelativeIKOp::GetBodyTransformInvariant(const UPhysicsAsse
 	return BodyRotInv * ElementScale3D * BodyTransform;
 }
 
+FVector FIKRetargetRelativeIKOp::CalcCompSpaceBlendLocation(const FIKRigGoal* InGoal, const FTransform& InCompBoneTfm) const
+{
+	if (!InGoal->bEnabled)
+	{
+		return InCompBoneTfm.GetLocation();
+	}
+
+	switch (InGoal->PositionSpace)
+	{
+		case EIKRigGoalSpace::Additive:
+		{
+			return FMath::Lerp(FVector::ZeroVector, InGoal->Position, InGoal->PositionAlpha) + InCompBoneTfm.GetLocation();
+		}
+		case EIKRigGoalSpace::Component:
+		{
+			return FMath::Lerp(InCompBoneTfm.GetLocation(), InGoal->Position, InGoal->PositionAlpha);
+		}
+		case EIKRigGoalSpace::World:
+		// TODO: How to support world-space goal blend?
+		default:
+		return InCompBoneTfm.GetLocation();
+	}
+}
+
 FKShapeElem* FIKRetargetRelativeIKOp::FindBodyShape(const UPhysicsAsset* PhysAsset, FName BoneName)
 {
 	int32 BodyIdx = PhysAsset->FindBodyIndex(BoneName);
diff --git a/Engine/Plugins/Experimental/Animation/RelativeIKOp/Source/RelativeIKOp/Public/RelativeIKOp.h b/Engine/Plugins/Experimental/Animation/RelativeIKOp/Source/RelativeIKOp/Public/RelativeIKOp.h
--- a/Engine/Plugins/Experimental/Animation/RelativeIKOp/Source/RelativeIKOp/Public/RelativeIKOp.h
+++ b/Engine/Plugins/Experimental/Animation/RelativeIKOp/Source/RelativeIKOp/Public/RelativeIKOp.h
@@ -219,6 +219,8 @@ private:
 	void UpdateCacheChainInfo(const TArray<FTransform>& TargetPose, FName BoneName);
 	void ComputeFeasibilityRange(FDoubleInterval& OutFeasibleRange, FName BoneName, FVector BoneChainOffset, FVector TargetRelativeToBoneV, FVector TargetRelativeToOpposeBoneV);
 
+	FVector CalcCompSpaceBlendLocation(const FIKRigGoal* InGoal, const FTransform& InCompBoneTfm) const;
+
 	static double CalcLineSphereIntersect(const FVector& Center, double Radius, const FVector& StartPoint, const FVector& EndPoint);
 
 	static FTransform CalcReferenceShapeScale3D(FKShapeElem* ShapeElem);
