RelativeIKOp: Use Gaussian smoothing with wider temporal support.

diff --git a/Engine/Plugins/Experimental/Animation/RelativeIKOp/Source/RelativeIKOp/Private/RelativeIKOp.cpp b/Engine/Plugins/Experimental/Animation/RelativeIKOp/Source/RelativeIKOp/Private/RelativeIKOp.cpp
--- a/Engine/Plugins/Experimental/Animation/RelativeIKOp/Source/RelativeIKOp/Private/RelativeIKOp.cpp
+++ b/Engine/Plugins/Experimental/Animation/RelativeIKOp/Source/RelativeIKOp/Private/RelativeIKOp.cpp
@@ -114,19 +114,35 @@ void FIKRetargetRelativeIKOp::Run(
 	int32 FramePairIdx = NumBodies * SampleIdx;
 
 	// Testing temporal smoothing of local points
-	const int32 MinSmoothIdx = FMath::Max(SampleIdx - Settings.TestTemporalSmoothingFrames, 0);
-	const int32 MaxSmoothIdx = FMath::Min(SampleIdx + Settings.TestTemporalSmoothingFrames,CachedNotifyInfo->NumSamples-1) ;
+	const int32 MinSmoothIdx = FMath::Max(SampleIdx - Settings.TestTemporalSmoothingRadius, 0);
+	const int32 MaxSmoothIdx = FMath::Min(SampleIdx + Settings.TestTemporalSmoothingRadius,CachedNotifyInfo->NumSamples-1) ;
 
 	TArray<FVector3f> SmoothedPoints;
 	SmoothedPoints.SetNum(NumBodies);
-	for (int BodyIdx = 0; BodyIdx < CachedNotifyInfo->BodyPairs.Num(); ++BodyIdx)
+	if (Settings.TestTemporalSmoothingRadius > 0)
 	{
-		SmoothedPoints[BodyIdx] = FVector3f::Zero();
-		for (int32 FrameIdx = MinSmoothIdx; FrameIdx <= MaxSmoothIdx; ++FrameIdx)
+		for (int BodyIdx = 0; BodyIdx < CachedNotifyInfo->BodyPairs.Num(); ++BodyIdx)
 		{
-			SmoothedPoints[BodyIdx] += CachedNotifyInfo->BodyPairsLocalReference[NumBodies * FrameIdx + BodyIdx];
+			SmoothedPoints[BodyIdx] = FVector3f::Zero();
+			float SumWeight = 0.0f;
+			for (int32 FrameIdx = MinSmoothIdx; FrameIdx <= MaxSmoothIdx; ++FrameIdx)
+			{
+				// Finite support Gaussian smoothing (could approximate w/ triangle or bake out smoothed data)
+				const float Sigma = 0.25f * static_cast<float>(Settings.TestTemporalSmoothingRadius);
+				const float SmoothDist = static_cast<float>(FrameIdx - SampleIdx);
+				const float SmoothWeight = FMath::Exp(-0.5f*(SmoothDist*SmoothDist) / (Sigma*Sigma));
+				SmoothedPoints[BodyIdx] += SmoothWeight * CachedNotifyInfo->BodyPairsLocalReference[NumBodies * FrameIdx + BodyIdx];
+				SumWeight += SmoothWeight;
+			}
+			SmoothedPoints[BodyIdx] /= SumWeight;
+		}
+	}
+	else
+	{
+		for (int BodyIdx = 0; BodyIdx < CachedNotifyInfo->BodyPairs.Num(); ++BodyIdx)
+		{
+			SmoothedPoints[BodyIdx] = CachedNotifyInfo->BodyPairsLocalReference[NumBodies * SampleIdx + BodyIdx];
 		}
-		SmoothedPoints[BodyIdx] /= static_cast<float>(MaxSmoothIdx - MinSmoothIdx + 1);
 	}
 
 	FDebugRelativeIKDrawInfo LocalDebugInfo;
diff --git a/Engine/Plugins/Experimental/Animation/RelativeIKOp/Source/RelativeIKOp/Public/RelativeIKOp.h b/Engine/Plugins/Experimental/Animation/RelativeIKOp/Source/RelativeIKOp/Public/RelativeIKOp.h
--- a/Engine/Plugins/Experimental/Animation/RelativeIKOp/Source/RelativeIKOp/Public/RelativeIKOp.h
+++ b/Engine/Plugins/Experimental/Animation/RelativeIKOp/Source/RelativeIKOp/Public/RelativeIKOp.h
@@ -168,8 +168,8 @@ struct FIKRetargetRelativeIKOpSettings : public FIKRetargetOpSettingsBase
 	bool bTestBodyRotInvariance = true;
 
 	// Frames of temporal smoothing of body pair verts (0 is no smoothing)
-	UPROPERTY(EditAnywhere, BlueprintReadWrite, Category = "Test Features", meta=(ClampMin = "0", ClampMax = "7"))
-	int32 TestTemporalSmoothingFrames = 3;
+	UPROPERTY(EditAnywhere, BlueprintReadWrite, Category = "Test Features", meta=(ClampMin = "0", ClampMax = "60"))
+	int32 TestTemporalSmoothingRadius = 15;
 };
 
 USTRUCT(BlueprintType, meta = (DisplayName = "Relative IK Goal Op"))
