RelativeIK: Update relative ik and body intersect ops to support new op debug draw display.

diff --git a/Engine/Plugins/Experimental/Animation/RelativeIKOp/Source/BodyIntersectIKOp/Private/BodyIntersectIKOp.cpp b/Engine/Plugins/Experimental/Animation/RelativeIKOp/Source/BodyIntersectIKOp/Private/BodyIntersectIKOp.cpp
--- a/Engine/Plugins/Experimental/Animation/RelativeIKOp/Source/BodyIntersectIKOp/Private/BodyIntersectIKOp.cpp
+++ b/Engine/Plugins/Experimental/Animation/RelativeIKOp/Source/BodyIntersectIKOp/Private/BodyIntersectIKOp.cpp
@@ -170,7 +170,7 @@ void FIKRetargetBodyIntersectIKOp::Run(
 		}
 		SetGoalPosFromCompSpace(IKGoal, TargetBoneTfm, CompGoalUpdate);
 #if WITH_EDITOR
-		if (Settings.bEnableDebugDraw)
+		if (Settings.bDebugDraw)
 		{
 			LocalDebugInfo.TestSpheres.Emplace(CompGoalUpdate, GoalRadius);
 		}
@@ -178,7 +178,7 @@ void FIKRetargetBodyIntersectIKOp::Run(
 	}
 
 #if WITH_EDITOR
-	if (Settings.bEnableDebugDraw)
+	if (Settings.bDebugDraw)
 	{
 		FScopeLock ScopeLock(&DebugDataMutex);
 		
@@ -408,11 +408,6 @@ void FIKRetargetBodyIntersectIKOp::DebugDraw(
 	const double InComponentScale,
 	const FIKRetargetDebugDrawState& InEditorState) const
 {
-	if (!Settings.bEnableDebugDraw)
-	{
-		return;
-	}
-	
 	FScopeLock ScopeLock(&DebugDataMutex);
 
 	if (!Settings.TargetPhysicsAssetOverride)
diff --git a/Engine/Plugins/Experimental/Animation/RelativeIKOp/Source/BodyIntersectIKOp/Public/BodyIntersectIKOp.h b/Engine/Plugins/Experimental/Animation/RelativeIKOp/Source/BodyIntersectIKOp/Public/BodyIntersectIKOp.h
--- a/Engine/Plugins/Experimental/Animation/RelativeIKOp/Source/BodyIntersectIKOp/Public/BodyIntersectIKOp.h
+++ b/Engine/Plugins/Experimental/Animation/RelativeIKOp/Source/BodyIntersectIKOp/Public/BodyIntersectIKOp.h
@@ -42,10 +42,6 @@ struct FIKRetargetBodyIntersectIKOpSettings : public FIKRetargetOpSettingsBase
 	// Physics bodies to do trivial intersection against
 	UPROPERTY(EditAnywhere, BlueprintReadWrite, Category = "Intersection", meta=(ReinitializeOnEdit))
 	TArray<FName> IntersectBodies;
-	
-	// Display source and target physics bodies for baked data
-	UPROPERTY(EditAnywhere, BlueprintReadWrite, Category = "Debug")
-	bool bEnableDebugDraw = false;
 };
 
 USTRUCT(BlueprintType, meta = (DisplayName = "Body Intersect Goals"))
@@ -102,6 +98,8 @@ public:
 		const double InComponentScale,
 		const FIKRetargetDebugDrawState& InEditorState) const override;
 
+	UE_API virtual bool HasDebugDrawing() override { return true; };
+
 private:
 	void DebugDrawPhysBody(FPrimitiveDrawInterface* InPDI, const FTransform& ParentTransform, double Scale, const FKShapeElem* ShapeElem, const FLinearColor& Color) const;
 	
diff --git a/Engine/Plugins/Experimental/Animation/RelativeIKOp/Source/RelativeIKOp/Private/RelativeIKOp.cpp b/Engine/Plugins/Experimental/Animation/RelativeIKOp/Source/RelativeIKOp/Private/RelativeIKOp.cpp
--- a/Engine/Plugins/Experimental/Animation/RelativeIKOp/Source/RelativeIKOp/Private/RelativeIKOp.cpp
+++ b/Engine/Plugins/Experimental/Animation/RelativeIKOp/Source/RelativeIKOp/Private/RelativeIKOp.cpp
@@ -282,7 +282,7 @@ void FIKRetargetRelativeIKOp::Run(
 
 #if WITH_EDITOR
 		// Debug draw verts
-		if (Settings.bEnableDebugDraw)
+		if (Settings.bDebugDraw)
 		{
 			LocalDebugInfo.SourcePairVerts.Add({SourcePrimaryCompV, SourceOtherCompV, false, FVector::ZeroVector, FVector::ZeroVector});
 			if (Settings.DebugFullRetargetPairBones.Contains(TargetPrimaryBone))
@@ -323,7 +323,7 @@ void FIKRetargetRelativeIKOp::Run(
 		FVector GoalPoint = FVector::ZeroVector;
 
 #if WITH_EDITOR
-		if (Settings.bEnableDebugDraw)
+		if (Settings.bDebugDraw)
 		{
 			LocalDebugInfo.TargetGoals.AddDefaulted();
 			LocalDebugInfo.TargetGoals.Last().Goal = FVector::ZeroVector;
@@ -342,7 +342,7 @@ void FIKRetargetRelativeIKOp::Run(
 			TotalWeight += FrameBoneWeights[Index];
 
 #if WITH_EDITOR
-			if (Settings.bEnableDebugDraw)
+			if (Settings.bDebugDraw)
 			{
 				LocalDebugInfo.TargetGoals.Last().PairTargets.Add(FrameBoneTargets[Index]);
 				LocalDebugInfo.PairRetargetInfo.Emplace(LocalTargetPairSpaces[Index]);
@@ -355,7 +355,7 @@ void FIKRetargetRelativeIKOp::Run(
 		{
 			GoalPoint /= TotalWeight;
 #if WITH_EDITOR
-			if (Settings.bEnableDebugDraw)
+			if (Settings.bDebugDraw)
 			{
 				LocalDebugInfo.TargetGoals.Last().Goal = GoalPoint;
 			}
@@ -367,7 +367,8 @@ void FIKRetargetRelativeIKOp::Run(
 			FVector PrevGoal = CalcCompSpaceBlendLocation(IKGoal, OutTargetGlobalPose[TargetBoneIdx]);
 			FVector BlendGoal = FMath::Lerp(PrevGoal, GoalPoint, GoalAlpha);
 
-			bool bSetGoals = !Settings.bEnableDebugDraw || !Settings.bDryRun;
+			// TODO: Remove dry run for release
+			bool bSetGoals = !Settings.bDebugDraw || !Settings.bDryRun;
 			if (bSetGoals)
 			{
 				IKGoal->bEnabled = true;
@@ -379,7 +380,7 @@ void FIKRetargetRelativeIKOp::Run(
 	}
 
 #if WITH_EDITOR
-	if (Settings.bEnableDebugDraw)
+	if (Settings.bDebugDraw)
 	{
 		FScopeLock ScopeLock(&DebugDataMutex);
 
@@ -859,7 +860,7 @@ void FIKRetargetRelativeIKOp::DebugDraw(
 	const double InComponentScale,
 	const FIKRetargetDebugDrawState& InEditorState) const
 {
-	if (!Settings.bEnableDebugDraw)
+	if (!Settings.bDebugDraw)
 	{
 		return;
 	}
diff --git a/Engine/Plugins/Experimental/Animation/RelativeIKOp/Source/RelativeIKOp/Public/RelativeIKOp.h b/Engine/Plugins/Experimental/Animation/RelativeIKOp/Source/RelativeIKOp/Public/RelativeIKOp.h
--- a/Engine/Plugins/Experimental/Animation/RelativeIKOp/Source/RelativeIKOp/Public/RelativeIKOp.h
+++ b/Engine/Plugins/Experimental/Animation/RelativeIKOp/Source/RelativeIKOp/Public/RelativeIKOp.h
@@ -127,10 +127,6 @@ struct FIKRetargetRelativeIKOpSettings : public FIKRetargetOpSettingsBase
 	UPROPERTY(BlueprintReadWrite, Category = "Parameters")
 	double RetargetSpringAlpha = 0.5;
 
-	// Enable debug displays
-	UPROPERTY(EditAnywhere, BlueprintReadWrite, Category = "Debug")
-	bool bEnableDebugDraw = false;
-
 	// Draw source and target body pair relationships for current animation
 	UPROPERTY(EditAnywhere, BlueprintReadWrite, Category = "Debug")
 	bool bDebugDrawBodyPairs = false;
@@ -239,6 +235,8 @@ public:
 		const double InComponentScale,
 		const FIKRetargetDebugDrawState& InEditorState) const override;
 
+	UE_API virtual bool HasDebugDrawing() override { return true; };
+
 private:
 	void DebugDrawBodyPairs(FPrimitiveDrawInterface* InPDI, const FTransform& BaseTransform, double Scale, const TArray<FDebugDrawBodyPair>& BodyPairs) const;
 	void DebugDrawGoalContributions(FPrimitiveDrawInterface* InPDI, const FTransform& BaseTransform, double Scale, const TArray<FDebugDrawTargetGoal>& GoalInfoLis) const;
