RelativeIKOp: Move temporal blend to function and clean up goal blend code for relik and body intersect ops.

diff --git a/Engine/Plugins/Experimental/Animation/RelativeIKOp/Source/BodyIntersectIKOp/Private/BodyIntersectIKOp.cpp b/Engine/Plugins/Experimental/Animation/RelativeIKOp/Source/BodyIntersectIKOp/Private/BodyIntersectIKOp.cpp
--- a/Engine/Plugins/Experimental/Animation/RelativeIKOp/Source/BodyIntersectIKOp/Private/BodyIntersectIKOp.cpp
+++ b/Engine/Plugins/Experimental/Animation/RelativeIKOp/Source/BodyIntersectIKOp/Private/BodyIntersectIKOp.cpp
@@ -134,13 +134,13 @@ void FIKRetargetBodyIntersectIKOp::Run(
 		const FTransform TargetBoneTfm = OutTargetGlobalPose[TargetBoneIdx];
 		
 		double GoalRadius = CalcShapeSmallRadius(GoalShape);
-		// TODO: Handle actual capsule offsets w/ a lever arm for rotation?
-		FVector CompGoalUpdate = GetGoalLocToCompSpace(IKGoal, TargetBoneTfm);
 		if (FMath::IsNearlyZero(GoalRadius))
 		{
 			continue;
 		}
-			
+
+		// TODO: Handle actual capsule offsets w/ a lever arm for rotation?
+		FVector CompGoalUpdate = GoalLocBlendCompSpace(IKGoal, TargetBoneTfm);
 		for (FName BodyName : Settings.IntersectBodies)
 		{
 			if (TargetBoneName == BodyName)
@@ -201,37 +201,36 @@ void FIKRetargetBodyIntersectIKOp::Run(
 #endif
 }
 
-FVector FIKRetargetBodyIntersectIKOp::GetGoalLocToCompSpace(const FIKRigGoal* Goal, const FTransform& BoneTfm) const
+FVector FIKRetargetBodyIntersectIKOp::GoalLocBlendCompSpace(const FIKRigGoal* Goal, const FTransform& BoneTfm) const
 {
 	switch (Goal->PositionSpace)
 	{
 	case EIKRigGoalSpace::Additive:
 		{
-			const FVector GoalPos = FMath::Lerp(FVector::ZeroVector, Goal->Position, Goal->PositionAlpha);
-			return BoneTfm.TransformPositionNoScale(GoalPos);
+			return FMath::Lerp(FVector::ZeroVector, Goal->Position, Goal->PositionAlpha) + BoneTfm.GetLocation();
 		}
 	case EIKRigGoalSpace::Component:
 		{
-			const FVector GoalPos = FMath::Lerp(BoneTfm.GetLocation(), Goal->Position, Goal->PositionAlpha);
-			return GoalPos;
+			return FMath::Lerp(BoneTfm.GetLocation(), Goal->Position, Goal->PositionAlpha);
 		}
 	case EIKRigGoalSpace::World:
-	default:
 		{
-			// TODO: Handle world-space case
-			return Goal->Position;
+			// We assume no World-space goals will be set using retarget stack
 		}
 	}
+
+	return BoneTfm.GetLocation();
 }
 
 void FIKRetargetBodyIntersectIKOp::SetGoalPosFromCompSpace(FIKRigGoal* Goal, const FTransform& BoneTfm, const FVector& CompSpaceLoc) const
 {
 	// TODO: Alpha passthrough and/or alpha ignore
+	// Should we just update position space to be component always?
 	switch (Goal->PositionSpace)
 	{
 	case EIKRigGoalSpace::Additive:
 		{
-			Goal->Position = BoneTfm.InverseTransformPositionNoScale(CompSpaceLoc);
+			Goal->Position = CompSpaceLoc - BoneTfm.GetLocation();
 			Goal->PositionAlpha = 1.0f;
 			break;
 		}
@@ -244,7 +243,7 @@ void FIKRetargetBodyIntersectIKOp::SetGoalPosFromCompSpace(FIKRigGoal* Goal, con
 	case EIKRigGoalSpace::World:
 	default:
 		{
-			// TODO: Handle world-space case
+			// We assume no World-space goals will be set using retarget stack
 			Goal->Position = CompSpaceLoc;
 			Goal->PositionAlpha = 1.0f;
 			break;
@@ -478,6 +477,8 @@ void FIKRetargetBodyIntersectIKOp::DebugDrawPhysBody(FPrimitiveDrawInterface* In
 
 void FIKRetargetBodyIntersectIKOp::ResetDebugInfo()
 {
+	FScopeLock ScopeLock(&DebugDataMutex);
+	
 	DebugDrawInfo.TargetIntersectTfms.Reset();
 	DebugDrawInfo.TestSpheres.Reset();
 }
diff --git a/Engine/Plugins/Experimental/Animation/RelativeIKOp/Source/BodyIntersectIKOp/Public/BodyIntersectIKOp.h b/Engine/Plugins/Experimental/Animation/RelativeIKOp/Source/BodyIntersectIKOp/Public/BodyIntersectIKOp.h
--- a/Engine/Plugins/Experimental/Animation/RelativeIKOp/Source/BodyIntersectIKOp/Public/BodyIntersectIKOp.h
+++ b/Engine/Plugins/Experimental/Animation/RelativeIKOp/Source/BodyIntersectIKOp/Public/BodyIntersectIKOp.h
@@ -79,7 +79,7 @@ private:
 	void UpdateCacheSkelInfo(const FTargetSkeleton& InTargetSkeleton);
 	void UpdateTargetBoneMap(FName TargetBoneName);
 
-	FVector GetGoalLocToCompSpace(const FIKRigGoal* Goal, const FTransform& BoneTfm) const;
+	FVector GoalLocBlendCompSpace(const FIKRigGoal* Goal, const FTransform& BoneTfm) const;
 	void SetGoalPosFromCompSpace(FIKRigGoal* Goal, const FTransform& BoneTfm, const FVector& CompSpaceLoc) const;
 	
 	static double CalcIntersectionDelta(const FVector&  TargetPoint, double TargetRadius, const FTransform& ShapeTfm, FKShapeElem* ShapeElem, FVector& OutDeltaDir);
@@ -123,7 +123,7 @@ private:
 	TMap<FName,int32> CacheTargetSkelIndices;
 };
 
-/* The blueprint/python API for editing a Stride Warping Op */
+/* The blueprint/python API for editing BodyIntersectIK Op */
 UCLASS(MinimalAPI, BlueprintType)
 class UIKRetargetBodyIntersectController : public UIKRetargetOpControllerBase
 {
@@ -132,12 +132,12 @@ class UIKRetargetBodyIntersectController : public UIKRetargetOpControllerBase
 public:
 
 	/* Get the current op settings as a struct.
-	 * @return FIKRetargetRelativeIKOpSettings struct with the current settings used by the op. */
+	 * @return FIKRetargetBodyIntersectIKOp struct with the current settings used by the op. */
 	UFUNCTION(BlueprintCallable, Category = Settings)
 	UE_API FIKRetargetBodyIntersectIKOpSettings GetSettings();
 
 	/* Set the solver settings. Input is a custom struct type for this solver.
-	 * @param InSettings a FIKRetargetRelativeIKOpSettings struct containing all the settings to apply to this op */
+	 * @param InSettings a FIKRetargetBodyIntersectIKOp struct containing all the settings to apply to this op */
 	UFUNCTION(BlueprintCallable, Category = Settings)
 	UE_API void SetSettings(FIKRetargetBodyIntersectIKOpSettings InSettings);
 };
diff --git a/Engine/Plugins/Experimental/Animation/RelativeIKOp/Source/RelativeIKOp/Private/RelativeIKOp.cpp b/Engine/Plugins/Experimental/Animation/RelativeIKOp/Source/RelativeIKOp/Private/RelativeIKOp.cpp
--- a/Engine/Plugins/Experimental/Animation/RelativeIKOp/Source/RelativeIKOp/Private/RelativeIKOp.cpp
+++ b/Engine/Plugins/Experimental/Animation/RelativeIKOp/Source/RelativeIKOp/Private/RelativeIKOp.cpp
@@ -108,42 +108,9 @@ void FIKRetargetRelativeIKOp::Run(
 	{
 		return;
 	}
-
-	int32 NumBodies = CachedNotifyInfo->BodyPairs.Num();
-	int32 SampleIdx = FMath::Min(static_cast<int32>(NotifyTime * SampleRate), CachedNotifyInfo->NumSamples-1);
-	int32 FramePairIdx = NumBodies * SampleIdx;
-
-	// Testing temporal smoothing of local points
-	const int32 MinSmoothIdx = FMath::Max(SampleIdx - Settings.TestTemporalSmoothingRadius, 0);
-	const int32 MaxSmoothIdx = FMath::Min(SampleIdx + Settings.TestTemporalSmoothingRadius,CachedNotifyInfo->NumSamples-1) ;
-
+	
 	TArray<FVector3f> SmoothedPoints;
-	SmoothedPoints.SetNum(NumBodies);
-	if (Settings.TestTemporalSmoothingRadius > 0)
-	{
-		for (int BodyIdx = 0; BodyIdx < CachedNotifyInfo->BodyPairs.Num(); ++BodyIdx)
-		{
-			SmoothedPoints[BodyIdx] = FVector3f::Zero();
-			float SumWeight = 0.0f;
-			for (int32 FrameIdx = MinSmoothIdx; FrameIdx <= MaxSmoothIdx; ++FrameIdx)
-			{
-				// Finite support Gaussian smoothing (could approximate w/ triangle or bake out smoothed data)
-				const float Sigma = 0.25f * static_cast<float>(Settings.TestTemporalSmoothingRadius);
-				const float SmoothDist = static_cast<float>(FrameIdx - SampleIdx);
-				const float SmoothWeight = FMath::Exp(-0.5f*(SmoothDist*SmoothDist) / (Sigma*Sigma));
-				SmoothedPoints[BodyIdx] += SmoothWeight * CachedNotifyInfo->BodyPairsLocalReference[NumBodies * FrameIdx + BodyIdx];
-				SumWeight += SmoothWeight;
-			}
-			SmoothedPoints[BodyIdx] /= SumWeight;
-		}
-	}
-	else
-	{
-		for (int BodyIdx = 0; BodyIdx < CachedNotifyInfo->BodyPairs.Num(); ++BodyIdx)
-		{
-			SmoothedPoints[BodyIdx] = CachedNotifyInfo->BodyPairsLocalReference[NumBodies * SampleIdx + BodyIdx];
-		}
-	}
+	ApplyTemporalSmoothing(SmoothedPoints, NotifyTime);
 
 	FDebugRelativeIKDrawInfo LocalDebugInfo;
 	
@@ -194,6 +161,11 @@ void FIKRetargetRelativeIKOp::Run(
 			continue;
 		}
 
+		if (!CacheTargetSkelIndices.Contains(TargetPrimaryBone) || !CacheTargetSkelIndices.Contains(TargetOtherBone))
+		{
+			continue;
+		}
+
 		int32 SourcePrimaryPoseIdx = CacheSourceSkelIndices[SourcePrimaryBone];
 		int32 SourceOtherPoseIdx = CacheSourceSkelIndices[SourceOtherBone];
 
@@ -206,7 +178,7 @@ void FIKRetargetRelativeIKOp::Run(
 		// Compute source transforms and source-space vertex distance
 		// const FVector3f& LocalPrimaryVert = CachedNotifyInfo->BodyPairsLocalReference[FramePairIdx + 2*PairIdx];
 		const FVector3f& LocalPrimaryVert = SmoothedPoints[2*PairIdx];
-		// TODO: Swap and clean these up if we bake in rot invaraint space
+		// TODO: Swap and clean these up if we bake in rot invariant space
 		// FTransform SourcePrimaryBodyTfm = GetBodyTransformInvariant(SourcePhysicsAsset, SourcePrimaryBone) * SourceScaleTfm;
 		FTransform SourcePrimaryBodyTfm = SourcePrimaryBodyRot * GetBodyTransformInvariant(SourcePhysicsAsset, SourcePrimaryBone) * SourceScaleTfm;
 		FTransform SourcePrimaryCompTfm = SourcePrimaryBodyTfm * InSourceGlobalPose[SourcePrimaryPoseIdx] * InvScaleTfm;
@@ -319,6 +291,11 @@ void FIKRetargetRelativeIKOp::Run(
 			continue;
 		}
 
+		if (!CacheTargetSkelIndices.Contains(TargetBoneName))
+		{
+			continue;
+		}
+
 		double TotalWeight = 0;
 		FVector GoalPoint = FVector::ZeroVector;
 
@@ -364,7 +341,7 @@ void FIKRetargetRelativeIKOp::Run(
 			// TODO: Play with total weight functions to alpha goal
 			float GoalAlpha = FMath::Min(1.0f, static_cast<float>(TotalWeight*FMath::Clamp(Settings.ContributionSumWeight, 0.0f, 1.0f)));
 			int32 TargetBoneIdx = CacheTargetSkelIndices[TargetBoneName];
-			FVector PrevGoal = CalcCompSpaceBlendLocation(IKGoal, OutTargetGlobalPose[TargetBoneIdx]);
+			FVector PrevGoal = GoalLocBlendCompSpace(IKGoal, OutTargetGlobalPose[TargetBoneIdx]);
 			FVector BlendGoal = FMath::Lerp(PrevGoal, GoalPoint, GoalAlpha);
 
 			// TODO: Remove dry run for release
@@ -453,6 +430,45 @@ void FIKRetargetRelativeIKOp::UpdateCacheChainInfo(const TArray<FTransform>& Tar
 	CacheChainStartMap.Emplace(BoneName, TargetPose[CacheBoneChains[BoneName]->BoneIndices[0]].GetTranslation());
 }
 
+void FIKRetargetRelativeIKOp::ApplyTemporalSmoothing(TArray<FVector3f>& SmoothedPoints, float Time) const
+{
+	int32 NumBodies = CachedNotifyInfo->BodyPairs.Num();
+	int32 SampleIdx = FMath::Clamp(static_cast<int32>(Time * SampleRate), 0, CachedNotifyInfo->NumSamples-1);
+
+	SmoothedPoints.SetNumUninitialized(NumBodies);
+
+	// Testing temporal smoothing of local points
+	const int32 MinSmoothIdx = FMath::Max(SampleIdx - Settings.TestTemporalSmoothingRadius, 0);
+	const int32 MaxSmoothIdx = FMath::Min(SampleIdx + Settings.TestTemporalSmoothingRadius,CachedNotifyInfo->NumSamples-1) ;
+	
+	if (Settings.TestTemporalSmoothingRadius > 0 && CachedNotifyInfo->NumSamples > 0)
+	{
+		for (int BodyIdx = 0; BodyIdx < CachedNotifyInfo->BodyPairs.Num(); ++BodyIdx)
+		{
+			SmoothedPoints[BodyIdx] = FVector3f::Zero();
+			float SumWeight = 0.0f;
+			for (int32 FrameIdx = MinSmoothIdx; FrameIdx <= MaxSmoothIdx; ++FrameIdx)
+			{
+				// Finite support Gaussian smoothing (could approximate w/ triangle or bake out smoothed data)
+				const float Sigma = 0.25f * static_cast<float>(Settings.TestTemporalSmoothingRadius);
+				const float SmoothDist = static_cast<float>(FrameIdx - SampleIdx);
+				const float SmoothWeight = FMath::Exp(-0.5f*(SmoothDist*SmoothDist) / (Sigma*Sigma));
+				SmoothedPoints[BodyIdx] += SmoothWeight * CachedNotifyInfo->BodyPairsLocalReference[NumBodies * FrameIdx + BodyIdx];
+				SumWeight += SmoothWeight;
+			}
+			SmoothedPoints[BodyIdx] /= SumWeight;
+		}
+	}
+	else
+	{
+		for (int BodyIdx = 0; BodyIdx < CachedNotifyInfo->BodyPairs.Num(); ++BodyIdx)
+		{
+			SmoothedPoints[BodyIdx] = CachedNotifyInfo->BodyPairsLocalReference[NumBodies * SampleIdx + BodyIdx];
+		}
+	}
+}
+
+
 void FIKRetargetRelativeIKOp::ComputeFeasibilityRange(FDoubleInterval& OutFeasibleRange, FName BoneName, FVector BoneChainOffset, FVector TargetRelativeToBoneV, FVector TargetRelativeToOpposeBoneV)
 {
 	OutFeasibleRange = {0.0, 1.0};
@@ -582,28 +598,31 @@ FTransform FIKRetargetRelativeIKOp::GetBodyTransformInvariant(const UPhysicsAsse
 	return BodyRotInv * ElementScale3D * BodyTransform;
 }
 
-FVector FIKRetargetRelativeIKOp::CalcCompSpaceBlendLocation(const FIKRigGoal* InGoal, const FTransform& InCompBoneTfm) const
+FVector FIKRetargetRelativeIKOp::GoalLocBlendCompSpace(const FIKRigGoal* Goal, const FTransform& BoneTfm) const
 {
-	if (!InGoal->bEnabled)
+	if (!Goal->bEnabled)
 	{
-		return InCompBoneTfm.GetLocation();
+		return BoneTfm.GetLocation();
 	}
 
-	switch (InGoal->PositionSpace)
+	switch (Goal->PositionSpace)
 	{
 		case EIKRigGoalSpace::Additive:
 		{
-			return FMath::Lerp(FVector::ZeroVector, InGoal->Position, InGoal->PositionAlpha) + InCompBoneTfm.GetLocation();
+			return FMath::Lerp(FVector::ZeroVector, Goal->Position, Goal->PositionAlpha) + BoneTfm.GetLocation();
 		}
 		case EIKRigGoalSpace::Component:
 		{
-			return FMath::Lerp(InCompBoneTfm.GetLocation(), InGoal->Position, InGoal->PositionAlpha);
+			return FMath::Lerp(BoneTfm.GetLocation(), Goal->Position, Goal->PositionAlpha);
 		}
 		case EIKRigGoalSpace::World:
-		// TODO: How to support world-space goal blend?
 		default:
-		return InCompBoneTfm.GetLocation();
+		{
+			// We assume no World-space goals will be set using retarget stack
+		}
 	}
+
+	return BoneTfm.GetLocation();
 }
 
 FKShapeElem* FIKRetargetRelativeIKOp::FindBodyShape(const UPhysicsAsset* PhysAsset, FName BoneName)
@@ -1064,6 +1083,8 @@ void FIKRetargetRelativeIKOp::DebugDrawPhysBody(FPrimitiveDrawInterface* InPDI,
 
 void FIKRetargetRelativeIKOp::ResetDebugInfo()
 {
+	FScopeLock ScopeLock(&DebugDataMutex);
+	
 	DebugDrawInfo.SourcePairVerts.Reset();
 	DebugDrawInfo.TargetPairVerts.Reset();
 	
diff --git a/Engine/Plugins/Experimental/Animation/RelativeIKOp/Source/RelativeIKOp/Public/RelativeIKOp.h b/Engine/Plugins/Experimental/Animation/RelativeIKOp/Source/RelativeIKOp/Public/RelativeIKOp.h
--- a/Engine/Plugins/Experimental/Animation/RelativeIKOp/Source/RelativeIKOp/Public/RelativeIKOp.h
+++ b/Engine/Plugins/Experimental/Animation/RelativeIKOp/Source/RelativeIKOp/Public/RelativeIKOp.h
@@ -168,7 +168,7 @@ struct FIKRetargetRelativeIKOpSettings : public FIKRetargetOpSettingsBase
 	int32 TestTemporalSmoothingRadius = 15;
 };
 
-USTRUCT(BlueprintType, meta = (DisplayName = "Relative IK Goal Op"))
+USTRUCT(BlueprintType, meta = (DisplayName = "Relative IK Goals"))
 struct FIKRetargetRelativeIKOp : public FIKRetargetOpBase
 {
 	GENERATED_BODY()
@@ -213,9 +213,11 @@ private:
 
 	FName ApplyBodyMap(FName BodyName);
 	void UpdateCacheChainInfo(const TArray<FTransform>& TargetPose, FName BoneName);
+	
 	void ComputeFeasibilityRange(FDoubleInterval& OutFeasibleRange, FName BoneName, FVector BoneChainOffset, FVector TargetRelativeToBoneV, FVector TargetRelativeToOpposeBoneV);
+	void ApplyTemporalSmoothing(TArray<FVector3f>& SmoothedPoints, float tim) const;
 
-	FVector CalcCompSpaceBlendLocation(const FIKRigGoal* InGoal, const FTransform& InCompBoneTfm) const;
+	FVector GoalLocBlendCompSpace(const FIKRigGoal* Goal, const FTransform& BoneTfm) const;
 
 	static double CalcLineSphereIntersect(const FVector& Center, double Radius, const FVector& StartPoint, const FVector& EndPoint);
 
