RelativeIKOp: Read notify info from either seq in AnimMontage or directly from AnimSeq playhead (max weight).

diff --git a/Engine/Plugins/Experimental/Animation/RelativeIKOp/Source/RelativeIKOp/Private/RelativeIKOp.cpp b/Engine/Plugins/Experimental/Animation/RelativeIKOp/Source/RelativeIKOp/Private/RelativeIKOp.cpp
--- a/Engine/Plugins/Experimental/Animation/RelativeIKOp/Source/RelativeIKOp/Private/RelativeIKOp.cpp
+++ b/Engine/Plugins/Experimental/Animation/RelativeIKOp/Source/RelativeIKOp/Private/RelativeIKOp.cpp
@@ -7,6 +7,7 @@
 #include "RelativeBodyAnimNotifies.h"
 #include "Animation/AnimInstance.h"
 #include "Animation/AnimMontage.h"
+#include "Animation/AnimSequence.h"
 #include "Components/SkeletalMeshComponent.h"
 #include "PhysicsEngine/PhysicsAsset.h"
 #include "PhysicsEngine/SkeletalBodySetup.h"
@@ -76,6 +77,7 @@ bool FIKRetargetRelativeIKOp::Initialize(
 
 	// Force update cache info if new anim montage found
 	MontageInstance = nullptr;
+	CacheSourceAnimSequence = nullptr;
 	
 	bIsInitialized = true;
 	return bIsInitialized;
@@ -95,22 +97,14 @@ void FIKRetargetRelativeIKOp::Run(
 	{
 		return;
 	}
-
-	if (!bIsInitialized || !MontageInstance || !CachedNotifyInfo)
-	{
-		return;
-	}
-
-	// TODO: Is this the correct way to get playtime (how to generalize to animseq-only notifies)
-	float MontageTime = MontageInstance->GetPosition();
-	float NotifyTime = MontageTime - NotifyAnimStart;
-	if (NotifyTime < 0.0f || MontageTime > NotifyAnimEnd)
+	
+	if (!bIsInitialized || !CacheSourceAnimSequence || !CachedNotifyInfo || AnimSeqPlayHead < 0.0f)
 	{
 		return;
 	}
 	
 	TArray<FVector3f> SmoothedPoints;
-	ApplyTemporalSmoothing(SmoothedPoints, NotifyTime);
+	ApplyTemporalSmoothing(SmoothedPoints, AnimSeqPlayHead);
 
 	FDebugRelativeIKDrawInfo LocalDebugInfo;
 	
@@ -644,27 +638,36 @@ double FIKRetargetRelativeIKOp::GetDistanceWeight(double Distance, double DistTh
 
 void FIKRetargetRelativeIKOp::AnimGraphPreUpdateMainThread(USkeletalMeshComponent& SourceMeshComponent, USkeletalMeshComponent& TargetMeshComponent)
 {
-	FAnimMontageInstance* OldMontageInstance = MontageInstance;
-	MontageInstance = nullptr;
 	if (!bIsInitialized)
 	{
 		return;
 	}
+	
+	const FAnimMontageInstance* PrevMontageInstance = MontageInstance;
+	MontageInstance = nullptr;
 
-	const UAnimInstance* SourceAnimInstance = SourceMeshComponent.GetAnimInstance();
+	UAnimInstance* SourceAnimInstance = SourceMeshComponent.GetAnimInstance();
 	if (!SourceAnimInstance)
 	{
 		return;
 	}
-
-	// TODO: Does this caching work?
+	
 	MontageInstance = SourceAnimInstance->GetActiveMontageInstance();
-	if (MontageInstance == OldMontageInstance)
+	if (MontageInstance)
 	{
+		// Early-out for montages
+		if (MontageInstance == PrevMontageInstance)
+		{
+			PreUpdateMontagePlayhead();
+			return;
+		}
+
+		ResetCacheNotifyInfo();
+		SetupRelativeIKNotifyInfoMontage();
 		return;
 	}
-	
-	PreUpdateRelativeIKNotifyInfo();
+
+	UpdateRelativeIKNotifyInfoAnimSeq(SourceAnimInstance);
 }
 
 void FIKRetargetRelativeIKOp::UpdateCacheBoneChains(const FIKRetargetProcessor& InProcessor, const FIKRetargetRunIKRigOp* ParentRigOp)
@@ -733,32 +736,52 @@ void FIKRetargetRelativeIKOp::UpdateCacheSkelInfo(const FRetargetSkeleton& InSou
 }
 
 
-void FIKRetargetRelativeIKOp::PreUpdateRelativeIKNotifyInfo()
+void FIKRetargetRelativeIKOp::PreUpdateMontagePlayhead()
 {
-	UpdateCacheNotifyInfo(nullptr, -1, -1);
 	if (!MontageInstance)
 	{
+		AnimSeqPlayHead = -1.0f;
 		return;
 	}
 
+	float MontageTime = MontageInstance->DeltaTimeRecord.GetPrevious() + MontageInstance->DeltaTimeRecord.Delta;
+	AnimSeqPlayHead = MontageTime - SegmentStartTime;
+	if (MontageTime > SegmentEndTime)
+	{
+		AnimSeqPlayHead = -1.0f;
+	}
+}
+
+void FIKRetargetRelativeIKOp::SetupRelativeIKNotifyInfoMontage()
+{
+	CacheSourceAnimSequence = nullptr;
 	UAnimMontage* Montage = MontageInstance->Montage;
 	if (!Montage)
 	{
 		return;
 	}
 
-	// HACK: Only grabs first physik dense notify
+	// HACK: Only grabs first relative notify
 	for (const FSlotAnimationTrack& SlotTrack : Montage->SlotAnimTracks)
 	{
 		for (const FAnimSegment& Segment : SlotTrack.AnimTrack.AnimSegments)
 		{
 			const TObjectPtr<UAnimSequenceBase>& SeqRef = Segment.GetAnimReference();
+			CacheSourceAnimSequence = Cast<UAnimSequence>(SeqRef.Get());
+			if (!CacheSourceAnimSequence)
+			{
+				continue;
+			}
+			
 			for (const FAnimNotifyEvent& NotifyEvent : SeqRef->Notifies)
 			{
-				URelativeBodyBakeAnimNotify* DenseInfo = Cast<URelativeBodyBakeAnimNotify>(NotifyEvent.Notify);
-				if (DenseInfo)
+				URelativeBodyBakeAnimNotify* BakeNotifyInfo = Cast<URelativeBodyBakeAnimNotify>(NotifyEvent.Notify);
+				if (BakeNotifyInfo)
 				{
-					UpdateCacheNotifyInfo(DenseInfo, Segment.AnimStartTime, Segment.AnimEndTime);
+					SegmentStartTime = Segment.AnimStartTime;
+					SegmentEndTime = Segment.AnimEndTime;
+					PreUpdateMontagePlayhead();
+					UpdateCacheNotifyInfo(BakeNotifyInfo);
 					return;
 				}
 			}
@@ -766,10 +789,81 @@ void FIKRetargetRelativeIKOp::PreUpdateRelativeIKNotifyInfo()
 	}
 }
 
-void FIKRetargetRelativeIKOp::UpdateCacheNotifyInfo(URelativeBodyBakeAnimNotify* NotifyInfo, float AnimStart, float AnimEnd)
+void FIKRetargetRelativeIKOp::UpdateRelativeIKNotifyInfoAnimSeq(UAnimInstance* SourceAnimInstance)
+{
+	// Get source tick records
+	TMap<FName, FAnimGroupInstance> SyncGroupTickRecords = SourceAnimInstance->GetSyncGroupMapRead();
+	TArray<FAnimTickRecord> UngroupedTickRecords = SourceAnimInstance->GetUngroupedActivePlayersRead();
+
+	// Check for active players with baked notify data, keep hightest blend amount
+	float MaxBlend = -1.0f;
+	const FAnimTickRecord* MaxNotifyRecord = nullptr;
+	auto ProcessTickRecords = [&MaxBlend,&MaxNotifyRecord](const TArray<FAnimTickRecord>& ActivePlayers)
+		{
+			for (const FAnimTickRecord& ActivePlayer : ActivePlayers)
+			{
+				if ( ActivePlayer.EffectiveBlendWeight <= MaxBlend )
+				{
+					continue;
+				}
+				if (UAnimSequence* Sequence = Cast<UAnimSequence>(ActivePlayer.SourceAsset))
+				{
+					// Used to determine which notifies should we considered "active".
+					check(ActivePlayer.DeltaTimeRecord->IsPreviousValid())
+					for (const FAnimNotifyEvent& NotifyEvent : Sequence->Notifies)
+					{
+						URelativeBodyBakeAnimNotify* BakedNotifyInfo = Cast<URelativeBodyBakeAnimNotify>(NotifyEvent.Notify);
+						if (BakedNotifyInfo)
+						{
+							MaxBlend = ActivePlayer.EffectiveBlendWeight;
+							MaxNotifyRecord = &ActivePlayer;
+						}
+					}
+				}
+			}
+		};
+	
+	// Find all the anim sequences (with sync groups) we updated this tick.
+	for (const TTuple<FName, FAnimGroupInstance>& SyncGroupPair : SyncGroupTickRecords)
+	{
+		ProcessTickRecords(SyncGroupPair.Value.ActivePlayers);
+	}
+	
+	// Find all the anim sequences (no sync groups) we updated this tick.
+	ProcessTickRecords(UngroupedTickRecords);
+	
+	if (!MaxNotifyRecord)
+	{
+		CacheSourceAnimSequence = nullptr;
+		ResetCacheNotifyInfo();
+		return;
+	}
+
+	AnimSeqPlayHead = MaxNotifyRecord->DeltaTimeRecord->GetPrevious() + MaxNotifyRecord->DeltaTimeRecord->Delta;
+	if (MaxNotifyRecord->SourceAsset == CacheSourceAnimSequence)
+	{
+		return;
+	}
+	
+	CacheSourceAnimSequence = Cast<UAnimSequence>(MaxNotifyRecord->SourceAsset);
+
+	// NOTE: Only grabs first bake notify
+	for (const FAnimNotifyEvent& NotifyEvent : CacheSourceAnimSequence->Notifies)
+	{
+		URelativeBodyBakeAnimNotify* BakedNotifyInfo = Cast<URelativeBodyBakeAnimNotify>(NotifyEvent.Notify);
+		if (BakedNotifyInfo)
+		{
+			UpdateCacheNotifyInfo(BakedNotifyInfo);
+			return;
+		}
+	}
+}
+
+
+void FIKRetargetRelativeIKOp::ResetCacheNotifyInfo()
 {
-	NotifyAnimStart = -1.0f;
-	NotifyAnimEnd = -1.0f;
+	SegmentStartTime = -1.0f;
+	SegmentEndTime = -1.0f;
 	CachedNotifyInfo = nullptr;
 
 	// Clear all cached data on update
@@ -778,15 +872,11 @@ void FIKRetargetRelativeIKOp::UpdateCacheNotifyInfo(URelativeBodyBakeAnimNotify*
 	CacheSourceDomainBones.Reset();
 	FrameBoneTargets.Reset();
 	FrameBoneWeights.Reset();
+}
 
-	if (!NotifyInfo)
-	{
-		return;
-	}
-
+void FIKRetargetRelativeIKOp::UpdateCacheNotifyInfo(URelativeBodyBakeAnimNotify* NotifyInfo)
+{
 	CachedNotifyInfo = NotifyInfo;
-	NotifyAnimStart = AnimStart;
-	NotifyAnimEnd = AnimEnd;
 
 	int32 NumBodies = CachedNotifyInfo->BodyPairs.Num();
 	FrameBoneTargets.SetNum(NumBodies);
diff --git a/Engine/Plugins/Experimental/Animation/RelativeIKOp/Source/RelativeIKOp/Public/RelativeIKOp.h b/Engine/Plugins/Experimental/Animation/RelativeIKOp/Source/RelativeIKOp/Public/RelativeIKOp.h
--- a/Engine/Plugins/Experimental/Animation/RelativeIKOp/Source/RelativeIKOp/Public/RelativeIKOp.h
+++ b/Engine/Plugins/Experimental/Animation/RelativeIKOp/Source/RelativeIKOp/Public/RelativeIKOp.h
@@ -14,6 +14,8 @@ struct FKShapeElem;
 struct FResolvedBoneChain;
 struct FAnimMontageInstance;
 struct FIKRetargetRunIKRigOp;
+class UAnimInstance;
+class UAnimSequence;
 class UPhysicsAsset;
 class URelativeBodyBakeAnimNotify;
 
@@ -208,8 +210,11 @@ private:
 	void UpdateAllBoneMaps(FName SourceBoneName);
 	void UpdateTargetBoneMap(FName TargetBoneName);
 	void UpdateCacheSkelInfo(const FRetargetSkeleton& InSourceSkeleton, const FTargetSkeleton& InTargetSkeleton);
-	void PreUpdateRelativeIKNotifyInfo();
-	void UpdateCacheNotifyInfo(URelativeBodyBakeAnimNotify* NotifyInfo, float AnimStart, float AnimEnd);
+	void PreUpdateMontagePlayhead();
+	void SetupRelativeIKNotifyInfoMontage();
+	void UpdateRelativeIKNotifyInfoAnimSeq(UAnimInstance* SourceAnimInstance);
+	void ResetCacheNotifyInfo();
+	void UpdateCacheNotifyInfo(URelativeBodyBakeAnimNotify* NotifyInfo);
 
 	FName ApplyBodyMap(FName BodyName);
 	void UpdateCacheChainInfo(const TArray<FTransform>& TargetPose, FName BoneName);
@@ -261,9 +266,9 @@ private:
 	TObjectPtr<const UIKRigDefinition> TargetIKRig;
 
 	// Cache notify-related data for faster weighted goal computations
+	TObjectPtr<const UAnimSequence> CacheSourceAnimSequence;
 	TObjectPtr<const URelativeBodyBakeAnimNotify> CachedNotifyInfo;
-	float NotifyAnimStart;
-	float NotifyAnimEnd;
+	float AnimSeqPlayHead;
 	float SampleRate;
 
 	// Cache source/target bone-names
@@ -285,8 +290,11 @@ private:
 	// Additional per-frame ik cache maps for length lookup
 	TMap<FName,double> CacheChainLengthMap;
 	TMap<FName,FVector> CacheChainStartMap;
-	
-	FAnimMontageInstance* MontageInstance = nullptr;
+
+	// Cached montage instance/segment info for early-out matching in preupdate
+	const FAnimMontageInstance* MontageInstance;
+	float SegmentStartTime;
+	float SegmentEndTime;
 };
 
 /* The blueprint/python API for editing a Stride Warping Op */
