RelativeIKOp: Reinit on body-map changes, fix for debug draw crash and disappearing target mesh because of bad valiues on 0-weight goal contributions

diff --git a/Engine/Plugins/Experimental/Animation/RelativeIKOp/Source/RelativeIKOp/Private/RelativeIKOp.cpp b/Engine/Plugins/Experimental/Animation/RelativeIKOp/Source/RelativeIKOp/Private/RelativeIKOp.cpp
--- a/Engine/Plugins/Experimental/Animation/RelativeIKOp/Source/RelativeIKOp/Private/RelativeIKOp.cpp
+++ b/Engine/Plugins/Experimental/Animation/RelativeIKOp/Source/RelativeIKOp/Private/RelativeIKOp.cpp
@@ -30,7 +30,7 @@ const UClass* FIKRetargetRelativeIKOpSettings::GetControllerType() const
 
 void FIKRetargetRelativeIKOpSettings::CopySettingsAtRuntime(const FIKRetargetOpSettingsBase* InSettingsToCopyFrom)
 {
-	static TArray<FName> PropertiesToIgnore = {"SourcePhysicsAssetOverride", "TargetPhysicsAssetOverride", "DomainBodies", "ContactBodies"};
+	static TArray<FName> PropertiesToIgnore = {"SourcePhysicsAssetOverride", "TargetPhysicsAssetOverride", "BodyMapping"};
 	FIKRetargetOpBase::CopyStructProperties(
 		FIKRetargetRelativeIKOpSettings::StaticStruct(),
 		InSettingsToCopyFrom,
@@ -136,6 +136,8 @@ void FIKRetargetRelativeIKOp::Run(
 		// TODO: Allow selecting a subset of body pairs in retargeter
 		FrameBoneWeights[2*PairIdx] = 0;
 		FrameBoneWeights[2*PairIdx + 1] = 0;
+		FrameBoneTargets[2*PairIdx] = FVector::ZeroVector;
+		FrameBoneTargets[2*PairIdx + 1] = FVector::ZeroVector;
 		
 		FName SourcePrimaryBone = CachedNotifyInfo->BodyPairs[2*PairIdx];
 		FName SourceOtherBone = CachedNotifyInfo->BodyPairs[2*PairIdx + 1];
@@ -691,28 +693,24 @@ void FIKRetargetRelativeIKOp::UpdateCacheBoneChains(const FIKRetargetProcessor&
 	}
 }
 
-void FIKRetargetRelativeIKOp::UpdateAllBoneMaps(FName SourceBoneName)
+void FIKRetargetRelativeIKOp::UpdateSourceBoneMap(FName SourceBoneName)
 {
-	FName TargetBoneName = ApplyBodyMap(SourceBoneName);
-	
 	int32 SourceIdx = SourceBoneNames.Find(SourceBoneName);
-	int32 TargetIdx = TargetBoneNames.Find(TargetBoneName);
-
-	if (SourceIdx == INDEX_NONE || TargetIdx == INDEX_NONE)
+	if (SourceIdx == INDEX_NONE)
 	{
+		CacheSourceSkelIndices.Remove(SourceBoneName);
 		return;
 	}
 
 	CacheSourceSkelIndices.Emplace(SourceBoneName, SourceIdx);
-	CacheTargetSkelIndices.Emplace(TargetBoneName, TargetIdx);
 }
 
 void FIKRetargetRelativeIKOp::UpdateTargetBoneMap(FName TargetBoneName)
 {
 	int32 TargetIdx = TargetBoneNames.Find(TargetBoneName);
-
 	if (TargetIdx == INDEX_NONE)
 	{
+		CacheTargetSkelIndices.Remove(TargetBoneName);
 		return;
 	}
 	
@@ -723,19 +721,11 @@ void FIKRetargetRelativeIKOp::UpdateCacheSkelInfo(const FRetargetSkeleton& InSou
 {
 	SourceBoneNames = InSourceSkeleton.BoneNames;
 	TargetBoneNames = InTargetSkeleton.BoneNames;
-	
-	for (FName SourceBodyName : Settings.DomainBodies)
-	{
-		UpdateAllBoneMaps(SourceBodyName);
-	}
 
-	for (FName SourceBodyName : Settings.ContactBodies)
-	{
-		UpdateAllBoneMaps(SourceBodyName);
-	}
+	CacheSourceSkelIndices.Reset();
+	CacheTargetSkelIndices.Reset();
 }
 
-
 void FIKRetargetRelativeIKOp::PreUpdateMontagePlayhead()
 {
 	if (!MontageInstance)
@@ -897,7 +887,9 @@ void FIKRetargetRelativeIKOp::UpdateCacheNotifyInfo(URelativeBodyBakeAnimNotify*
 		CacheSourceEffectBones.Add(SourceBodyBone);
 		CacheSourceBodyEffectVertIdx.AddDefaulted();
 
-		UpdateAllBoneMaps(SourceBodyBone);
+		FName TargetBodyBone = ApplyBodyMap(SourceBodyBone);
+		UpdateSourceBoneMap(SourceBodyBone);
+		UpdateTargetBoneMap(TargetBodyBone);
 	}
 
 	// Cache bone-bone effectors for creating weighted goals
diff --git a/Engine/Plugins/Experimental/Animation/RelativeIKOp/Source/RelativeIKOp/Public/RelativeIKOp.h b/Engine/Plugins/Experimental/Animation/RelativeIKOp/Source/RelativeIKOp/Public/RelativeIKOp.h
--- a/Engine/Plugins/Experimental/Animation/RelativeIKOp/Source/RelativeIKOp/Public/RelativeIKOp.h
+++ b/Engine/Plugins/Experimental/Animation/RelativeIKOp/Source/RelativeIKOp/Public/RelativeIKOp.h
@@ -90,16 +90,8 @@ struct FIKRetargetRelativeIKOpSettings : public FIKRetargetOpSettingsBase
 	TObjectPtr<UPhysicsAsset> TargetPhysicsAssetOverride;
 
 	// Physics body source -> target name mapping 
-	UPROPERTY(EditAnywhere, BlueprintReadWrite, Category = "Source Target Mapping")
+	UPROPERTY(EditAnywhere, BlueprintReadWrite, Category = "Source Target Mapping", meta=(ReinitializeOnEdit))
 	TMap<FName,FName> BodyMapping;
-
-	// Physics body names that should maintain relative spatial relationships with ContactBodies 
-	UPROPERTY(EditAnywhere, BlueprintReadWrite, Category = "Body Pairs")
-	TArray<FName> DomainBodies;
-
-	// Physics body names that should be moved (IK goals) to maintain relative spatial relationships
-	UPROPERTY(EditAnywhere, BlueprintReadWrite, Category = "Body Pairs", meta=(ReinitializeOnEdit))
-	TArray<FName> ContactBodies;
 	
 	// Maximum distance for which body pair info is baked
 	UPROPERTY(EditAnywhere, BlueprintReadWrite, Category = "Parameters", meta=(ClampMin = "0.1"))
@@ -207,7 +199,7 @@ struct FIKRetargetRelativeIKOp : public FIKRetargetOpBase
 
 private:
 	void UpdateCacheBoneChains(const FIKRetargetProcessor& InProcessor, const FIKRetargetRunIKRigOp* ParentRigOp);
-	void UpdateAllBoneMaps(FName SourceBoneName);
+	void UpdateSourceBoneMap(FName SourceBoneName);
 	void UpdateTargetBoneMap(FName TargetBoneName);
 	void UpdateCacheSkelInfo(const FRetargetSkeleton& InSourceSkeleton, const FTargetSkeleton& InTargetSkeleton);
 	void PreUpdateMontagePlayhead();
