RelativeIKOp: Correctly handle different source/target bone orientations in relative ik. Also fixed sphere intersection check for feasibility ranges.

diff --git a/Engine/Plugins/Experimental/Animation/RelativeIKOp/Source/RelativeIKOp/Private/RelativeIKOp.cpp b/Engine/Plugins/Experimental/Animation/RelativeIKOp/Source/RelativeIKOp/Private/RelativeIKOp.cpp
--- a/Engine/Plugins/Experimental/Animation/RelativeIKOp/Source/RelativeIKOp/Private/RelativeIKOp.cpp
+++ b/Engine/Plugins/Experimental/Animation/RelativeIKOp/Source/RelativeIKOp/Private/RelativeIKOp.cpp
@@ -113,12 +113,13 @@ void FIKRetargetRelativeIKOp::Run(
 	LocalDebugInfo.TargetPairVerts.Reserve(CachedNotifyInfo->bBodyPairsIsParentDominates.Num());
 
 	TArray<FDebugRelativeTargetPairSpace> LocalTargetPairSpaces;
-	LocalTargetPairSpaces.SetNum(CachedNotifyInfo->BodyPairs.Num());
+	LocalTargetPairSpaces.SetNumZeroed(CachedNotifyInfo->BodyPairs.Num());
 #endif //WITH_EDITOR
 
 	const double SourceScale = InProcessor.GetSourceScaleFactor();
 	const double InvSourceScale = (Settings.bIgnoreSourceScale) ? 1.0 / InProcessor.GetSourceScaleFactor() : 1.0;
 
+	// NOTE: Uniform scaling so these should compose fine w/ FTransform
 	const FTransform SourceScaleTfm(FQuat::Identity, FVector::Zero(), FVector(SourceScale));
 	const FTransform InvScaleTfm(FQuat::Identity, FVector::Zero(), FVector(InvSourceScale));
 	
@@ -145,8 +146,6 @@ void FIKRetargetRelativeIKOp::Run(
 		FName TargetPrimaryBone = ApplyBodyMap(SourcePrimaryBone);
 		FName TargetOtherBone = ApplyBodyMap(SourceOtherBone);
 
-		// TODO: Early-out if missing body in current source/target PhysAsset
-		// Check this in update?
 		if (TargetPrimaryBone == NAME_None || TargetOtherBone == NAME_None)
 		{
 			continue;
@@ -162,56 +161,83 @@ void FIKRetargetRelativeIKOp::Run(
 			continue;
 		}
 
+		if (!HasBoneDelta(TargetPrimaryBone) || !HasBoneDelta(TargetOtherBone))
+		{
+			continue;
+		}
+
 		int32 SourcePrimaryPoseIdx = CacheSourceSkelIndices[SourcePrimaryBone];
 		int32 SourceOtherPoseIdx = CacheSourceSkelIndices[SourceOtherBone];
 
-		// ---- TESTING ONLY - Flip these once notify baking is rotation invariant
+		// ---- TESTING ONLY - Remove these once notify baking is rotation invariant
 		FTransform SourcePrimaryBodyRot = GetBodyRotation(SourcePhysicsAsset, SourcePrimaryBone);
 		FTransform SourceOtherBodyRot = GetBodyRotation(SourcePhysicsAsset, SourceOtherBone);
-		FTransform TargetPrimaryBodyRot = (Settings.bTestBodyRotInvariance) ? FTransform::Identity : SourcePrimaryBodyRot.Inverse() * GetBodyRotation(TargetPhysicsAsset, TargetPrimaryBone);
-		FTransform TargetOtherBodyRot = (Settings.bTestBodyRotInvariance) ? FTransform::Identity : SourceOtherBodyRot.Inverse() * GetBodyRotation(TargetPhysicsAsset, TargetOtherBone);
+		FTransform TargetPrimaryBodyRot = SourcePrimaryBodyRot;
+		FTransform TargetOtherBodyRot = SourceOtherBodyRot;
 		
 		// Compute source transforms and source-space vertex distance
 		// const FVector3f& LocalPrimaryVert = CachedNotifyInfo->BodyPairsLocalReference[FramePairIdx + 2*PairIdx];
+		// const FVector3f& LocalOtherVert = CachedNotifyInfo->BodyPairsLocalReference[FramePairIdx + 2*PairIdx + 1];
 		const FVector3f& LocalPrimaryVert = SmoothedPoints[2*PairIdx];
-		// TODO: Swap and clean these up if we bake in rot invariant space
-		// FTransform SourcePrimaryBodyTfm = GetBodyTransformInvariant(SourcePhysicsAsset, SourcePrimaryBone) * SourceScaleTfm;
-		FTransform SourcePrimaryBodyTfm = SourcePrimaryBodyRot * GetBodyTransformInvariant(SourcePhysicsAsset, SourcePrimaryBone) * SourceScaleTfm;
+		const FVector3f& LocalOtherVert = SmoothedPoints[2*PairIdx+1];
+
+		// TODO: Clean this up if we bake in rot invariant space
+		FVector SourceBoneBakePrimaryVert = (SourcePrimaryBodyRot).TransformPosition(FVector(LocalPrimaryVert)); 
+		FVector SourceBoneBakeOtherVert = (SourceOtherBodyRot).TransformPosition(FVector(LocalOtherVert));
+		// NOTE: Need to use FMatrix to apply oriented scale properly
+		FMatrix SourceBodyPrimaryScale = GetBodyOrientedScale(SourcePhysicsAsset, SourcePrimaryBone);
+		FMatrix SourceBodyOtherScale = GetBodyOrientedScale(SourcePhysicsAsset, SourceOtherBone);
+		FVector SourceBodyPrimaryV = SourceBodyPrimaryScale.TransformPosition(SourceBoneBakePrimaryVert);
+		FVector SourceBodyOtherV = SourceBodyOtherScale.TransformPosition(SourceBoneBakeOtherVert);
+		
+		FTransform SourcePrimaryBodyTfm = GetBodyTranslation(SourcePhysicsAsset, SourcePrimaryBone) * SourceScaleTfm;
 		FTransform SourcePrimaryCompTfm = SourcePrimaryBodyTfm * InSourceGlobalPose[SourcePrimaryPoseIdx] * InvScaleTfm;
-		FVector SourcePrimaryCompV = SourcePrimaryCompTfm.TransformPosition(FVector(LocalPrimaryVert));
+		FVector SourcePrimaryCompV = SourcePrimaryCompTfm.TransformPosition(SourceBodyPrimaryV);
 
-		// const FVector3f& LocalOtherVert = CachedNotifyInfo->BodyPairsLocalReference[FramePairIdx + 2*PairIdx + 1];
-		const FVector3f& LocalOtherVert = SmoothedPoints[2*PairIdx+1];
-		FTransform SourceOtherBodyTfm = SourceOtherBodyRot * GetBodyTransformInvariant(SourcePhysicsAsset, SourceOtherBone) * SourceScaleTfm;
+		FTransform SourceOtherBodyTfm = GetBodyTranslation(SourcePhysicsAsset, SourceOtherBone) * SourceScaleTfm;
 		FTransform SourceOtherCompTfm = SourceOtherBodyTfm * InSourceGlobalPose[SourceOtherPoseIdx] * InvScaleTfm;
-		FVector SourceOtherCompV = SourceOtherCompTfm.TransformPosition(FVector(LocalOtherVert));
+		FVector SourceOtherCompV = SourceOtherCompTfm.TransformPosition(SourceBodyOtherV);
 
 		double SourceVDist = FVector::Distance(SourcePrimaryCompV, SourceOtherCompV);
 		double LinearWeight = FMath::Clamp(1.0 - (SourceVDist / Settings.DistanceThreshold), 0.0, 1.0);
 		double Weight = GetDistanceWeight(SourceVDist, Settings.DistanceThreshold, Settings.DistanceFade);
 
 		// Compute Source pair verts as represented in alternate pair body local space
-		FVector SourcePrimaryOtherV = SourceOtherCompTfm.InverseTransformPosition(SourcePrimaryCompV);
-		FVector SourceOtherPrimaryV = SourcePrimaryCompTfm.InverseTransformPosition(SourceOtherCompV);
-
-		// Get average of retarget verts
+		FVector SourcePrimaryOtherV = SourceBodyOtherScale.InverseTransformPosition(SourceOtherCompTfm.InverseTransformPosition(SourcePrimaryCompV));
+		FVector SourceOtherPrimaryV = SourceBodyPrimaryScale.InverseTransformPosition(SourcePrimaryCompTfm.InverseTransformPosition(SourceOtherCompV));
+		
+		//---------------------------------
+		// Get verts retargeted using physics body oriented scale
 		int32 TargetPrimaryPoseIdx = CacheTargetSkelIndices[TargetPrimaryBone];
 		int32 TargetOtherPoseIdx = CacheTargetSkelIndices[TargetOtherBone];
-		// TODO: Swap and clean these up if we bake in rot invaraint space
-		// FTransform TargetPrimaryBodyTfm = GetBodyTransformInvariant(TargetPhysicsAsset, TargetPrimaryBone);
-		FTransform TargetPrimaryBodyTfm = SourcePrimaryBodyRot * TargetPrimaryBodyRot * GetBodyTransformInvariant(TargetPhysicsAsset, TargetPrimaryBone);
-		FTransform TargetOtherBodyTfm = SourceOtherBodyRot * TargetOtherBodyRot * GetBodyTransformInvariant(TargetPhysicsAsset, TargetOtherBone);
+
+		FTransform TargetPrimaryBoneDelta = GetTargetBoneDelta(TargetPrimaryBone);
+		FTransform TargetOtherBoneDelta = GetTargetBoneDelta(TargetOtherBone);
+		
+		FVector TargetBoneBakePrimaryVert = (TargetPrimaryBodyRot * TargetPrimaryBoneDelta).TransformPosition(FVector(LocalPrimaryVert)); 
+		FVector TargetBoneBakeOtherVert = (TargetOtherBodyRot * TargetOtherBoneDelta).TransformPosition(FVector(LocalOtherVert));
+
+		// NOTE: Need to use FMatrix to apply oriented scale properly
+		FMatrix TargetBodyPrimaryScale = GetBodyOrientedScale(TargetPhysicsAsset, TargetPrimaryBone);
+		FMatrix TargetBodyOtherScale = GetBodyOrientedScale(TargetPhysicsAsset, TargetOtherBone);
+		FVector TargetBodyPrimaryV = TargetBodyPrimaryScale.TransformPosition(TargetBoneBakePrimaryVert);
+		FVector TargetBodyOtherV = TargetBodyOtherScale.TransformPosition(TargetBoneBakeOtherVert);
+		
+		FTransform TargetPrimaryBodyTfm = GetBodyTranslation(TargetPhysicsAsset, TargetPrimaryBone);
+		FTransform TargetOtherBodyTfm = GetBodyTranslation(TargetPhysicsAsset, TargetOtherBone);
 
 		FTransform TargetPrimaryCompTfm = TargetPrimaryBodyTfm * OutTargetGlobalPose[TargetPrimaryPoseIdx];
 		FTransform TargetOtherCompTfm = TargetOtherBodyTfm * OutTargetGlobalPose[TargetOtherPoseIdx];
 
 		//TODO: Play with/think about offset strategy for these points
-		FVector TargetPrimaryCompV_P = TargetPrimaryCompTfm.TransformPosition(FVector(LocalPrimaryVert));
-		FVector TargetPrimaryCompV_O = TargetOtherCompTfm.TransformPosition(SourcePrimaryOtherV);
+		FMatrix TestPrimaryScale = (Settings.bTestRetargetScale) ? TargetBodyPrimaryScale : SourceBodyPrimaryScale;
+		FMatrix TestOtherScale = (Settings.bTestRetargetScale) ? TargetBodyOtherScale : SourceBodyOtherScale;
+		FVector TargetPrimaryCompV_P = TargetPrimaryCompTfm.TransformPosition(TargetBodyPrimaryV);
+		FVector TargetPrimaryCompV_O = TargetOtherCompTfm.TransformPosition(TestOtherScale.TransformPosition(TargetOtherBoneDelta.TransformPosition(SourcePrimaryOtherV)));
 		FVector TargetPrimaryBoneOffset = OutTargetGlobalPose[TargetPrimaryPoseIdx].GetTranslation() - TargetPrimaryCompV_P;
 
-		FVector TargetOtherCompV_O = TargetOtherCompTfm.TransformPosition(FVector(LocalOtherVert));
-		FVector TargetOtherCompV_P = TargetPrimaryCompTfm.TransformPosition(SourceOtherPrimaryV);
+		FVector TargetOtherCompV_O = TargetOtherCompTfm.TransformPosition(TargetBodyOtherV);
+		FVector TargetOtherCompV_P = TargetPrimaryCompTfm.TransformPosition(TestPrimaryScale.TransformPosition(TargetPrimaryBoneDelta.TransformPosition(SourceOtherPrimaryV)));
 		FVector TargetOtherBoneOffset = OutTargetGlobalPose[TargetOtherPoseIdx].GetTranslation() - TargetOtherCompV_O;
 
 		// double SpringAlpha = (Settings.bTestDistContactAlpha) ? Weight : RetargetSpringAlpha;
@@ -372,14 +398,51 @@ void FIKRetargetRelativeIKOp::Run(
 			}
 
 			int32 SourcePoseIdx = CacheSourceSkelIndices[SourceBoneName];
-			DebugDrawInfo.SourceBodyInfo.Add({SourceBoneName, SourceScaleTfm * InSourceGlobalPose[SourcePoseIdx]});
+			DebugDrawInfo.SourceBodyInfo.Emplace(SourceBoneName, SourceScaleTfm * InSourceGlobalPose[SourcePoseIdx], FTransform::Identity);
 		}
 
 		// Cache debug info for target domain bodies
 		DebugDrawInfo.TargetBodyInfo.Reset(CacheSourceDomainBones.Num());
 		for (FName SourceBoneName : CacheSourceDomainBones)
 		{
-			// TODO: Only show domain bodies? 
+			FName TargetBoneName = ApplyBodyMap(SourceBoneName);
+			if (!CacheTargetSkelIndices.Contains(TargetBoneName))
+			{
+				continue;
+			}
+
+			FTransform TargetBoneDelta = GetTargetBoneDelta(TargetBoneName);
+			int32 TargetPoseIdx = CacheTargetSkelIndices[TargetBoneName];
+			DebugDrawInfo.TargetBodyInfo.Emplace(TargetBoneName, OutTargetGlobalPose[TargetPoseIdx], TargetBoneDelta);
+		}
+
+		DebugDrawInfo.SourceTfmInfo.Reset(CacheSourceEffectBones.Num());
+		for (FName SourceBoneName : CacheSourceEffectBones)
+		{
+			if (!CacheSourceSkelIndices.Contains(SourceBoneName))
+			{
+				continue;
+			}
+
+			int32 SourcePoseIdx = CacheSourceSkelIndices[SourceBoneName];
+
+			// NOTE: Need to use FMatrix to apply oriented scale properly
+			FMatrix SourceBodyRot = GetBodyRotation(SourcePhysicsAsset, SourceBoneName).ToMatrixNoScale();
+			FMatrix SourceBodyScale = GetBodyOrientedScale(SourcePhysicsAsset, SourceBoneName);
+			FMatrix SourceBodyTfm = SourceBodyRot * SourceBodyScale * GetBodyTranslation(SourcePhysicsAsset, SourceBoneName).ToMatrixNoScale();
+			FMatrix SourceCompTfm = SourceBodyTfm * InSourceGlobalPose[SourcePoseIdx].ToMatrixWithScale();
+
+			FVector SourceBodyCenter = InSourceGlobalPose[SourcePoseIdx].TransformPosition(GetBodyTranslation(SourcePhysicsAsset, SourceBoneName).GetLocation());
+			FVector SourceBodyX = SourceCompTfm.TransformPosition(FVector::UnitX());
+			FVector SourceBodyY = SourceCompTfm.TransformPosition(FVector::UnitY());
+			FVector SourceBodyZ = SourceCompTfm.TransformPosition(FVector::UnitZ());
+			
+			DebugDrawInfo.SourceTfmInfo.Emplace(SourceBodyCenter, SourceBodyX, SourceBodyY, SourceBodyZ);
+		}
+		
+		DebugDrawInfo.TargetTfmInfo.Reset(CacheSourceEffectBones.Num());
+		for (FName SourceBoneName : CacheSourceEffectBones)
+		{
 			FName TargetBoneName = ApplyBodyMap(SourceBoneName);
 			if (!CacheTargetSkelIndices.Contains(TargetBoneName))
 			{
@@ -387,7 +450,21 @@ void FIKRetargetRelativeIKOp::Run(
 			}
 			
 			int32 TargetPoseIdx = CacheTargetSkelIndices[TargetBoneName];
-			DebugDrawInfo.TargetBodyInfo.Add({TargetBoneName, OutTargetGlobalPose[TargetPoseIdx]});
+
+			// NOTE: Need to use FMatrix to apply oriented scale properly
+			FMatrix TargetBoneDelta = GetTargetBoneDelta(TargetBoneName).ToMatrixNoScale();
+			
+			FMatrix SourceBodyRot = GetBodyRotation(SourcePhysicsAsset, SourceBoneName).ToMatrixNoScale();
+			FMatrix TargetBodyScale = GetBodyOrientedScale(TargetPhysicsAsset, TargetBoneName);
+			FMatrix TargetBodyTfm = SourceBodyRot * TargetBoneDelta * TargetBodyScale * GetBodyTranslation(TargetPhysicsAsset, TargetBoneName).ToMatrixNoScale();
+			FMatrix TargetCompTfm = TargetBodyTfm * OutTargetGlobalPose[TargetPoseIdx].ToMatrixWithScale();
+
+			FVector TargetBodyCenter = OutTargetGlobalPose[TargetPoseIdx].TransformPosition(GetBodyTranslation(TargetPhysicsAsset, TargetBoneName).GetLocation());
+			FVector TargetBodyX = TargetCompTfm.TransformPosition(FVector::UnitX());
+			FVector TargetBodyY = TargetCompTfm.TransformPosition(FVector::UnitY());
+			FVector TargetBodyZ = TargetCompTfm.TransformPosition(FVector::UnitZ());
+			
+			DebugDrawInfo.TargetTfmInfo.Emplace(TargetBodyCenter, TargetBodyX, TargetBodyY, TargetBodyZ);
 		}
 	}
 #endif
@@ -403,6 +480,20 @@ FName FIKRetargetRelativeIKOp::ApplyBodyMap(FName BodyName)
 	return Settings.BodyMapping[BodyName];
 }
 
+bool FIKRetargetRelativeIKOp::HasBoneDelta(FName TargetBoneName) const
+{
+	return CacheMapBoneSourceTargetTfm.Contains(TargetBoneName);
+}
+
+const FTransform& FIKRetargetRelativeIKOp::GetTargetBoneDelta(FName TargetBoneName) const
+{
+	if (!CacheMapBoneSourceTargetTfm.Contains(TargetBoneName))
+	{
+		return FTransform::Identity;
+	}
+	return CacheMapBoneSourceTargetTfm[TargetBoneName];
+}
+
 void FIKRetargetRelativeIKOp::UpdateCacheChainInfo(const TArray<FTransform>& TargetPose, FName BoneName)
 {
 	if (!CacheBoneChains.Contains(BoneName))
@@ -478,40 +569,32 @@ void FIKRetargetRelativeIKOp::ComputeFeasibilityRange(FDoubleInterval& OutFeasib
 		return;
 	}
 	
-			
-	// FVector OffsetChainStart = ChainStartMap[OtherBone] - RelativeOtherOffset - TargetOtherBoneOffset;
 	FVector OffsetChainStart = CacheChainStartMap[BoneName] - BoneChainOffset;
-	// TODO: Better bias range
 	double CheckRadius = FMath::Max(CacheChainLengthMap[BoneName] + Settings.FeasibilityLengthBias, 0.0);
 			
 	double ChainDistStart = FVector::Distance(TargetRelativeToBoneV, OffsetChainStart);
 	double ChainDistEnd = FVector::Distance(TargetRelativeToOpposeBoneV, OffsetChainStart);
 
-	bool UpperRange = ChainDistStart > ChainDistEnd; 
+	FDoubleInterval IntersectRange;
 	if (FMath::Max(ChainDistStart, ChainDistEnd) <= CheckRadius)
 	{
 		// Whole range is feasible
-		OutFeasibleRange = FDoubleInterval(0.0,1.0);
+		OutFeasibleRange = {0.0,1.0};
 	}
-	else if (FMath::Min(ChainDistStart, ChainDistEnd) >= CheckRadius)
+	else if (!CalcLineSphereIntersect(IntersectRange, OffsetChainStart, CheckRadius, TargetRelativeToBoneV, TargetRelativeToOpposeBoneV))
 	{
-		// None of the range is feasible so use lower-dist endpoint
-		// TODO: This is an approximation as we could have a smaller feasible range that doesn't include either endpoint
-		double RangePt = (UpperRange) ? 1.0 : 0.0;
-		OutFeasibleRange = FDoubleInterval(RangePt,RangePt);
+		// No feasible range, pick closer of the endpoints
+		double RangeEndPt = (ChainDistStart > ChainDistEnd) ? 1.0 : 0.0;
+		OutFeasibleRange = {RangeEndPt,RangeEndPt};
 	}
 	else
 	{
-		// Figure out alpha range of feasible points along contribution target line
-		double RangePt = CalcLineSphereIntersect(OffsetChainStart, CheckRadius, TargetRelativeToBoneV, TargetRelativeToOpposeBoneV);
-		double MinVal = (UpperRange) ? RangePt : 0.0;
-		double MaxVal = (UpperRange) ? 1.0 : RangePt;
-
-		OutFeasibleRange = FDoubleInterval(MinVal, MaxVal);
+		// Use sphere intersection range, clamped on [0,1]
+		OutFeasibleRange = {FMath::Clamp(IntersectRange.Min, 0.0,1.0), FMath::Clamp(IntersectRange.Max, 0.0,1.0)};
 	}
 }
 
-double FIKRetargetRelativeIKOp::CalcLineSphereIntersect(const FVector& Center, double Radius, const FVector& StartPoint, const FVector& EndPoint)
+bool FIKRetargetRelativeIKOp::CalcLineSphereIntersect(FDoubleInterval& OutRange, const FVector& Center, double Radius, const FVector& StartPoint, const FVector& EndPoint)
 {
 	// NOTE: This function assumes two intersection points and returns positive ratio along line from start point
 	// Used for getting the endpoint of feasible range
@@ -521,8 +604,17 @@ double FIKRetargetRelativeIKOp::CalcLineSphereIntersect(const FVector& Center, d
 	double a = LineVec.SquaredLength();
 	double b = FVector::DotProduct(LineVec, LocPt);
 	double c = LocPt.SquaredLength() - Radius*Radius;
+	if (b*b - a*c < 0.0)
+	{
+		return false;
+	}
+
+	double B = FMath::Sqrt(b*b - a*c);
+	double MaxRange = (-b + B) / a;
+	double MinRange = (-b - B) / a;
+	OutRange = {MinRange, MaxRange};
 
-	return (-b + FMath::Sqrt(b*b - a*c)) / a;
+	return true;
 }
 
 FTransform FIKRetargetRelativeIKOp::CalcReferenceShapeScale3D(FKShapeElem* ShapeElem)
@@ -580,7 +672,7 @@ FTransform FIKRetargetRelativeIKOp::GetBodyRotation(const UPhysicsAsset* PhysAss
 	return FTransform(BodyTransform.GetRotation());
 }
 
-FTransform FIKRetargetRelativeIKOp::GetBodyTransformInvariant(const UPhysicsAsset* PhysAsset, FName BoneName)
+FTransform FIKRetargetRelativeIKOp::GetBodyTranslation(const UPhysicsAsset* PhysAsset, FName BoneName)
 {
 	FKShapeElem* ShapeElem = FindBodyShape(PhysAsset, BoneName);
 	if (!ShapeElem)
@@ -589,9 +681,21 @@ FTransform FIKRetargetRelativeIKOp::GetBodyTransformInvariant(const UPhysicsAsse
 	}
 	
 	FTransform BodyTransform = ShapeElem->GetTransform();
-	FTransform BodyRotInv(BodyTransform.GetRotation().Inverse());
-	FTransform ElementScale3D = CalcReferenceShapeScale3D(ShapeElem);
-	return BodyRotInv * ElementScale3D * BodyTransform;
+	return FTransform(BodyTransform.GetTranslation());
+}
+
+FMatrix FIKRetargetRelativeIKOp::GetBodyOrientedScale(const UPhysicsAsset* PhysAsset, FName BoneName)
+{
+	FKShapeElem* ShapeElem = FindBodyShape(PhysAsset, BoneName);
+	if (!ShapeElem)
+	{
+		return FMatrix::Identity;
+	}
+	
+	FTransform BodyTransform = ShapeElem->GetTransform();
+	FQuat BodyRotation = BodyTransform.GetRotation();
+	FMatrix ElementScale3D = CalcReferenceShapeScale3D(ShapeElem).ToMatrixWithScale();
+	return BodyRotation.Inverse().ToMatrix() * ElementScale3D * BodyRotation.ToMatrix();
 }
 
 FVector FIKRetargetRelativeIKOp::GoalLocBlendCompSpace(const FIKRigGoal* Goal, const FTransform& BoneTfm) const
@@ -717,11 +821,31 @@ void FIKRetargetRelativeIKOp::UpdateTargetBoneMap(FName TargetBoneName)
 	CacheTargetSkelIndices.Emplace(TargetBoneName, TargetIdx);
 }
 
+void FIKRetargetRelativeIKOp::UpdateBoneMapTfm(FName SourceBoneName, FName TargetBoneName)
+{
+	int32 SourceIdx = SourceBoneNames.Find(SourceBoneName);
+	int32 TargetIdx = TargetBoneNames.Find(TargetBoneName);
+	if (SourceIdx == INDEX_NONE || TargetIdx == INDEX_NONE)
+	{
+		CacheMapBoneSourceTargetTfm.Remove(TargetBoneName);
+		return;
+	}
+
+	FQuat SourceRot = CacheSourceBoneInitTfm[SourceIdx].GetRotation();
+	FQuat TargetInvRot = CacheTargetBoneInitTfm[TargetIdx].GetRotation().Inverse();
+
+	FTransform TargetDelta = FTransform(TargetInvRot*SourceRot);
+	CacheMapBoneSourceTargetTfm.Emplace(TargetBoneName,TargetDelta);
+}
+
 void FIKRetargetRelativeIKOp::UpdateCacheSkelInfo(const FRetargetSkeleton& InSourceSkeleton, const FTargetSkeleton& InTargetSkeleton)
 {
 	SourceBoneNames = InSourceSkeleton.BoneNames;
 	TargetBoneNames = InTargetSkeleton.BoneNames;
 
+	CacheSourceBoneInitTfm = InSourceSkeleton.RetargetPoses.GetGlobalRetargetPose();
+	CacheTargetBoneInitTfm = InTargetSkeleton.RetargetPoses.GetGlobalRetargetPose();
+
 	CacheSourceSkelIndices.Reset();
 	CacheTargetSkelIndices.Reset();
 }
@@ -734,7 +858,8 @@ void FIKRetargetRelativeIKOp::PreUpdateMontagePlayhead()
 		return;
 	}
 
-	float MontageTime = MontageInstance->DeltaTimeRecord.GetPrevious() + MontageInstance->DeltaTimeRecord.Delta;
+	// TODO: DeltaTimeRecord doesn't seem to work for montages when scrubbing
+	float MontageTime = MontageInstance->GetPosition();
 	AnimSeqPlayHead = MontageTime - SegmentStartTime;
 	if (MontageTime > SegmentEndTime)
 	{
@@ -857,6 +982,7 @@ void FIKRetargetRelativeIKOp::ResetCacheNotifyInfo()
 	CachedNotifyInfo = nullptr;
 
 	// Clear all cached data on update
+	CacheMapBoneSourceTargetTfm.Reset();
 	CacheSourceBodyEffectVertIdx.Reset();
 	CacheSourceEffectBones.Reset();
 	CacheSourceDomainBones.Reset();
@@ -890,6 +1016,7 @@ void FIKRetargetRelativeIKOp::UpdateCacheNotifyInfo(URelativeBodyBakeAnimNotify*
 		FName TargetBodyBone = ApplyBodyMap(SourceBodyBone);
 		UpdateSourceBoneMap(SourceBodyBone);
 		UpdateTargetBoneMap(TargetBodyBone);
+		UpdateBoneMapTfm(SourceBodyBone, TargetBodyBone);
 	}
 
 	// Cache bone-bone effectors for creating weighted goals
@@ -993,6 +1120,18 @@ void FIKRetargetRelativeIKOp::DebugDraw(
 		}
 	}
 
+	if (Settings.bDebugDrawBodyTransforms)
+	{
+		if (Settings.SourcePhysicsAssetOverride)
+		{
+			DebugDrawBodyCoords(InPDI, SourceNoScaleTfm, SourceScale, Settings.SourcePhysicsAssetOverride, DebugDrawInfo.SourceTfmInfo);
+		}
+		if (Settings.TargetPhysicsAssetOverride)
+		{
+			DebugDrawBodyCoords(InPDI, InComponentTransform, InComponentScale, Settings.TargetPhysicsAssetOverride, DebugDrawInfo.TargetTfmInfo);
+		}
+	}
+
 	if (Settings.bDebugDrawRetargetVertAverages)
 	{
 		// Show retarget contribution average with offsets
@@ -1008,13 +1147,15 @@ void FIKRetargetRelativeIKOp::DebugDraw(
 
 void FIKRetargetRelativeIKOp::DebugDrawBodyPairs(FPrimitiveDrawInterface* InPDI, const FTransform& BaseTransform, double Scale, const TArray<FDebugDrawBodyPair>& BodyPairs) const
 {
+	const float MaxScale = 1.0f;
+	const float MarkerSize = FMath::Min(MaxScale, static_cast<float>(Scale));
 	for (const FDebugDrawBodyPair& Pair : BodyPairs)
 	{
 		FTransform TfmBodyA_A = FTransform(Pair.PosA_A) * BaseTransform;
 		FTransform TfmBodyB_B = FTransform(Pair.PosB_B) * BaseTransform;
 
-		IKRigDebugRendering::DrawWireCube(InPDI, TfmBodyA_A, FLinearColor::Blue,static_cast<float>(Scale),static_cast<float>(Scale));
-		IKRigDebugRendering::DrawWireCube(InPDI, TfmBodyB_B, FLinearColor::Red,static_cast<float>(Scale),static_cast<float>(Scale));
+		IKRigDebugRendering::DrawWireCube(InPDI, TfmBodyA_A, FLinearColor::Blue,MarkerSize,MarkerSize);
+		IKRigDebugRendering::DrawWireCube(InPDI, TfmBodyB_B, FLinearColor::Red,MarkerSize,MarkerSize);
 
 		// Possibly draw full retarget debug info
 		if (Pair.bFullPos)
@@ -1022,8 +1163,8 @@ void FIKRetargetRelativeIKOp::DebugDrawBodyPairs(FPrimitiveDrawInterface* InPDI,
 			FTransform TfmBodyA_B = FTransform(Pair.PosA_B) * BaseTransform;
 			FTransform TfmBodyB_A = FTransform(Pair.PosB_A) * BaseTransform;
 
-			IKRigDebugRendering::DrawWireCube(InPDI, TfmBodyA_B, FLinearColor::Black,static_cast<float>(Scale),static_cast<float>(Scale));
-			IKRigDebugRendering::DrawWireCube(InPDI, TfmBodyB_A, FLinearColor::Yellow,static_cast<float>(Scale),static_cast<float>(Scale));
+			IKRigDebugRendering::DrawWireCube(InPDI, TfmBodyA_B, FLinearColor::Black,MarkerSize,MarkerSize);
+			IKRigDebugRendering::DrawWireCube(InPDI, TfmBodyB_A, FLinearColor::Yellow,MarkerSize,MarkerSize);
 
 			// Draw direct relationships between bodies in retargeted A/B spaces
 			DrawDashedLine(InPDI, TfmBodyA_A.GetLocation(), TfmBodyB_A.GetLocation(), FLinearColor::Black, 1, SDPG_Foreground);
@@ -1044,6 +1185,8 @@ void FIKRetargetRelativeIKOp::DebugDrawBodyPairs(FPrimitiveDrawInterface* InPDI,
 void FIKRetargetRelativeIKOp::DebugDrawGoalContributions(FPrimitiveDrawInterface* InPDI, const FTransform& BaseTransform, double Scale, const TArray<FDebugDrawTargetGoal>& GoalInfoLis) const
 {
 	// Show IK Goals (and contribs)
+	const float MaxScale = 1.0f;
+	const float MarkerSize = FMath::Min(MaxScale, static_cast<float>(Scale));
 	for (const FDebugDrawTargetGoal& Goal : GoalInfoLis)
 	{
 		FTransform GoalTfm = FTransform(Goal.Goal) * BaseTransform;
@@ -1053,21 +1196,22 @@ void FIKRetargetRelativeIKOp::DebugDrawGoalContributions(FPrimitiveDrawInterface
 			FTransform ContribTfm = FTransform(v) * BaseTransform;
 			FVector ContribPt = ContribTfm.GetTranslation();
 			DrawDashedLine(InPDI, ContribPt, GoalPoint, FLinearColor::Red, 1, SDPG_Foreground);
-			IKRigDebugRendering::DrawWireCube(InPDI, ContribTfm, FLinearColor::White,static_cast<float>(Scale),static_cast<float>(Scale));
+			IKRigDebugRendering::DrawWireCube(InPDI, ContribTfm, FLinearColor::White,MarkerSize,MarkerSize);
 		}
-		IKRigDebugRendering::DrawWireCube(InPDI, GoalTfm, FLinearColor::Yellow,static_cast<float>(Scale),static_cast<float>(Scale));
+		IKRigDebugRendering::DrawWireCube(InPDI, GoalTfm, FLinearColor::Yellow,MarkerSize,MarkerSize);
 	}
 }
 
 void FIKRetargetRelativeIKOp::DebugDrawPairVertRetarget(FPrimitiveDrawInterface* InPDI, const FTransform& BaseTransform, double Scale, const TArray<FDebugRelativeTargetPairSpace>& RetargetPairList) const
 {
-	const float MarkerSize = 0.5f;
+	const float MaxScale = 1.0f;
+	const float MarkerSize = 0.5f * FMath::Min(MaxScale, static_cast<float>(Scale));
 	for (const FDebugRelativeTargetPairSpace& Info : RetargetPairList)
 	{
 		FTransform V1Tfm = FTransform(Info.PairRangeStart) * BaseTransform;
 		FTransform V2Tfm = FTransform(Info.PairRangeEnd) * BaseTransform;
-		IKRigDebugRendering::DrawWireCube(InPDI, V1Tfm, FLinearColor::Red,static_cast<float>(MarkerSize*Scale),static_cast<float>(MarkerSize*Scale));
-		IKRigDebugRendering::DrawWireCube(InPDI, V2Tfm, FLinearColor::Red,static_cast<float>(MarkerSize*Scale),static_cast<float>(MarkerSize*Scale));
+		IKRigDebugRendering::DrawWireCube(InPDI, V1Tfm, FLinearColor::Red,MarkerSize,MarkerSize);
+		IKRigDebugRendering::DrawWireCube(InPDI, V2Tfm, FLinearColor::Red,MarkerSize,MarkerSize);
 
 		FVector V1 = V1Tfm.GetTranslation();
 		FVector V2 = V2Tfm.GetTranslation();
@@ -1076,24 +1220,21 @@ void FIKRetargetRelativeIKOp::DebugDrawPairVertRetarget(FPrimitiveDrawInterface*
 		DrawDashedLine(InPDI, V1, V2, FLinearColor::Black, 1, SDPG_Foreground);
 
 		FTransform VAvgTfm = FTransform(Info.PairTarget) * BaseTransform;
-		IKRigDebugRendering::DrawWireCube(InPDI, VAvgTfm, FLinearColor::Gray,static_cast<float>(MarkerSize*Scale),static_cast<float>(MarkerSize*Scale));
+		IKRigDebugRendering::DrawWireCube(InPDI, VAvgTfm, FLinearColor::Gray,MarkerSize,MarkerSize);
 
 		FVector FeasibleMinV =  FMath::Lerp(Info.PairRangeStart, Info.PairRangeEnd, Info.FeasibleRange.Min);
 		FTransform VMinTfm = FTransform(FeasibleMinV) * BaseTransform;
-		IKRigDebugRendering::DrawWireCube(InPDI, VMinTfm, FLinearColor::Yellow,static_cast<float>(MarkerSize*Scale),static_cast<float>(MarkerSize*Scale));
+		IKRigDebugRendering::DrawWireCube(InPDI, VMinTfm, FLinearColor::Yellow,MarkerSize,MarkerSize);
 
 		FVector FeasibleMaxV =  FMath::Lerp(Info.PairRangeStart, Info.PairRangeEnd, Info.FeasibleRange.Max);
 		FTransform VMaxTfm = FTransform(FeasibleMaxV) * BaseTransform;
-		IKRigDebugRendering::DrawWireCube(InPDI, VMaxTfm, FLinearColor::Green,static_cast<float>(MarkerSize*Scale),static_cast<float>(MarkerSize*Scale));
+		IKRigDebugRendering::DrawWireCube(InPDI, VMaxTfm, FLinearColor::Green,MarkerSize,MarkerSize);
 		// DrawDashedLine(InPDI, VOffset, VOffset + t.Get<3>(), FLinearColor::Black, 1, SDPG_Foreground);
 	}
 }
 
 void FIKRetargetRelativeIKOp::DebugDrawBodies(FPrimitiveDrawInterface* InPDI, const FTransform& BaseTransform, double Scale, UPhysicsAsset* PhysAsset, const TArray<FDebugDrawBodyInfo>& PhysBodies) const
 {
-	const float CoordSize = 10.0f;
-	const float CoordThickness = 0.5f;
-	
 	for (const FDebugDrawBodyInfo& BodyInfo : PhysBodies)
 	{
 		const FKShapeElem* ShapeElem = FindBodyShape(PhysAsset, BodyInfo.BodyName);
@@ -1104,19 +1245,6 @@ void FIKRetargetRelativeIKOp::DebugDrawBodies(FPrimitiveDrawInterface* InPDI, co
 		
 		FTransform CompTfm = BodyInfo.BoneTfm * BaseTransform;
 
-		FTransform BodyRot = (Settings.bTestBodyRotInvariance) ? FTransform::Identity : GetBodyRotation(SourcePhysicsAsset, BodyInfo.BodyName);
-		FTransform BodyTfm = BodyRot * GetBodyTransformInvariant(PhysAsset, BodyInfo.BodyName) * CompTfm;
-
-		// Draw body coord system:
-		DrawCoordinateSystem(
-			InPDI,
-			BodyTfm.GetLocation(),
-			BodyTfm.GetRotation().Rotator(),
-			static_cast<float>(CoordSize * Scale),
-			SDPG_Foreground,
-			static_cast<float>(CoordThickness * Scale));
-
-
 		// Draw physics body
 		DebugDrawPhysBody(InPDI, CompTfm, Scale, ShapeElem, FLinearColor::Yellow);
 	}
@@ -1162,6 +1290,28 @@ void FIKRetargetRelativeIKOp::DebugDrawPhysBody(FPrimitiveDrawInterface* InPDI,
 }
 
 
+void FIKRetargetRelativeIKOp::DebugDrawBodyCoords(FPrimitiveDrawInterface* InPDI, const FTransform& BaseTransform, double Scale, UPhysicsAsset* PhysAsset, const TArray<FDebugBodyTfmInfo>& BodyTfms) const
+{
+	// Stop scale from making coords too large
+	const float MaxScale = 1.0f;
+	const float CoordSize = 10.0f * FMath::Min(MaxScale, static_cast<float>(Scale));
+	const float CoordThickness = 0.5f * FMath::Min(MaxScale, static_cast<float>(Scale));
+
+	for (const FDebugBodyTfmInfo& BodyTfm : BodyTfms)
+	{
+
+		FVector StartV = BaseTransform.TransformPosition(BodyTfm.Center);
+		FVector CoordX = BaseTransform.TransformPosition(BodyTfm.TfmX);
+		FVector CoordY = BaseTransform.TransformPosition(BodyTfm.TfmY);
+		FVector CoordZ = BaseTransform.TransformPosition(BodyTfm.TfmZ);
+		
+		InPDI->DrawLine(StartV, CoordX, FLinearColor::Red, SDPG_Foreground, CoordThickness);
+		InPDI->DrawLine(StartV, CoordY, FLinearColor::Green, SDPG_Foreground, CoordThickness);
+		InPDI->DrawLine(StartV, CoordZ, FLinearColor::Blue, SDPG_Foreground, CoordThickness);
+	}
+}
+
+
 
 void FIKRetargetRelativeIKOp::ResetDebugInfo()
 {
@@ -1172,6 +1322,9 @@ void FIKRetargetRelativeIKOp::ResetDebugInfo()
 	
 	DebugDrawInfo.SourceBodyInfo.Reset();
 	DebugDrawInfo.TargetBodyInfo.Reset();
+
+	DebugDrawInfo.SourceBodyInfo.Reset();
+	DebugDrawInfo.TargetTfmInfo.Reset();
 	
 	DebugDrawInfo.PairRetargetInfo.Reset();
 	DebugDrawInfo.TargetGoals.Reset();
diff --git a/Engine/Plugins/Experimental/Animation/RelativeIKOp/Source/RelativeIKOp/Public/RelativeIKOp.h b/Engine/Plugins/Experimental/Animation/RelativeIKOp/Source/RelativeIKOp/Public/RelativeIKOp.h
--- a/Engine/Plugins/Experimental/Animation/RelativeIKOp/Source/RelativeIKOp/Public/RelativeIKOp.h
+++ b/Engine/Plugins/Experimental/Animation/RelativeIKOp/Source/RelativeIKOp/Public/RelativeIKOp.h
@@ -53,6 +53,15 @@ struct FDebugDrawBodyInfo
 {
 	FName BodyName;
 	FTransform BoneTfm;
+	FTransform RetargetTfm = FTransform::Identity;
+};
+
+struct FDebugBodyTfmInfo
+{
+	FVector Center;
+	FVector TfmX;
+	FVector TfmY;
+	FVector TfmZ;
 };
 
 // Used to hold all debug draw info for convenience
@@ -66,6 +75,10 @@ struct FDebugRelativeIKDrawInfo
 	TArray<FDebugDrawBodyInfo> SourceBodyInfo;
 	TArray<FDebugDrawBodyInfo> TargetBodyInfo;
 
+	// Body transform vecs
+	TArray<FDebugBodyTfmInfo> SourceTfmInfo;
+	TArray<FDebugBodyTfmInfo> TargetTfmInfo;
+
 	// Debug info for retarget pairs
 	TArray<FDebugRelativeTargetPairSpace> PairRetargetInfo;
 	// Final goal points with contributions
@@ -141,6 +154,10 @@ struct FIKRetargetRelativeIKOpSettings : public FIKRetargetOpSettingsBase
 	UPROPERTY(EditAnywhere, BlueprintReadWrite, Category = "Debug")
 	bool bDebugDrawPhysicsBodies = false;
 
+	// Show body-local transform setup (like baked verts)
+	UPROPERTY(EditAnywhere, BlueprintReadWrite, Category = "Debug")
+	bool bDebugDrawBodyTransforms = false;
+
 	// Run op and display debug info but DON'T update IK Goals
 	UPROPERTY(EditAnywhere, BlueprintReadWrite, Category = "Debug")
 	bool bDryRun = false;
@@ -153,9 +170,9 @@ struct FIKRetargetRelativeIKOpSettings : public FIKRetargetOpSettingsBase
 	UPROPERTY(EditAnywhere, BlueprintReadWrite, Category = "Test Features")
 	bool bTestFeasibilityWeight = true;
 
-	// Use rotation-invariant space for physics body components on source/target
+	// Use target (default) when computing pair alternate representations
 	UPROPERTY(EditAnywhere, BlueprintReadWrite, Category = "Test Features")
-	bool bTestBodyRotInvariance = true;
+	bool bTestRetargetScale = true;
 
 	// Frames of temporal smoothing of body pair verts (0 is no smoothing)
 	UPROPERTY(EditAnywhere, BlueprintReadWrite, Category = "Test Features", meta=(ClampMin = "0", ClampMax = "60"))
@@ -201,6 +218,7 @@ private:
 	void UpdateCacheBoneChains(const FIKRetargetProcessor& InProcessor, const FIKRetargetRunIKRigOp* ParentRigOp);
 	void UpdateSourceBoneMap(FName SourceBoneName);
 	void UpdateTargetBoneMap(FName TargetBoneName);
+	void UpdateBoneMapTfm(FName SourceBoneName, FName TargetBoneName);
 	void UpdateCacheSkelInfo(const FRetargetSkeleton& InSourceSkeleton, const FTargetSkeleton& InTargetSkeleton);
 	void PreUpdateMontagePlayhead();
 	void SetupRelativeIKNotifyInfoMontage();
@@ -209,6 +227,8 @@ private:
 	void UpdateCacheNotifyInfo(URelativeBodyBakeAnimNotify* NotifyInfo);
 
 	FName ApplyBodyMap(FName BodyName);
+	bool HasBoneDelta(FName TargetBoneName) const;
+	const FTransform& GetTargetBoneDelta(FName TargetBoneName) const;
 	void UpdateCacheChainInfo(const TArray<FTransform>& TargetPose, FName BoneName);
 	
 	void ComputeFeasibilityRange(FDoubleInterval& OutFeasibleRange, FName BoneName, FVector BoneChainOffset, FVector TargetRelativeToBoneV, FVector TargetRelativeToOpposeBoneV);
@@ -216,11 +236,15 @@ private:
 
 	FVector GoalLocBlendCompSpace(const FIKRigGoal* Goal, const FTransform& BoneTfm) const;
 
-	static double CalcLineSphereIntersect(const FVector& Center, double Radius, const FVector& StartPoint, const FVector& EndPoint);
+	static bool CalcLineSphereIntersect(FDoubleInterval& OutRange, const FVector& Center, double Radius, const FVector& StartPoint, const FVector& EndPoint);
 
 	static FTransform CalcReferenceShapeScale3D(FKShapeElem* ShapeElem);
+	// Get local body rotation relative to parent bone
 	static FTransform GetBodyRotation(const UPhysicsAsset* PhysAsset, FName BoneName);
-	static FTransform GetBodyTransformInvariant(const UPhysicsAsset* PhysAsset, FName BoneName);
+	// Get body origin relative to parent bone
+	static FTransform GetBodyTranslation(const UPhysicsAsset* PhysAsset, FName BoneName);
+	// Get body oriented (non-axis aligned scaling) relative to parent bone
+	static FMatrix GetBodyOrientedScale(const UPhysicsAsset* PhysAsset, FName BoneName);
 	static FKShapeElem* FindBodyShape(const UPhysicsAsset* PhysAsset, FName BoneName);
 	
 	static double GetDistanceWeight(double Distance, double DistThreshold, double ScalarFade = 200.0);
@@ -242,6 +266,7 @@ private:
 	void DebugDrawPairVertRetarget(FPrimitiveDrawInterface* InPDI, const FTransform& BaseTransform, double Scale, const TArray<FDebugRelativeTargetPairSpace>& RetargetPairList) const;
 	void DebugDrawBodies(FPrimitiveDrawInterface* InPDI, const FTransform& BaseTransform, double Scale, UPhysicsAsset* PhysAsset, const TArray<FDebugDrawBodyInfo>& PhysBodies) const;
 	void DebugDrawPhysBody(FPrimitiveDrawInterface* InPDI, const FTransform& ParentTransform, double Scale, const FKShapeElem* ShapeElem, const FLinearColor& Color) const;
+	void DebugDrawBodyCoords(FPrimitiveDrawInterface* InPDI, const FTransform& BaseTransform, double Scale, UPhysicsAsset* PhysAsset, const TArray<FDebugBodyTfmInfo>& BodyTfms) const;
 	
 	void ResetDebugInfo();
 
@@ -275,6 +300,9 @@ private:
 	TArray<FVector> FrameBoneTargets;
 
 	// Cache skel and IK goal info
+	TArray<FTransform> CacheSourceBoneInitTfm;
+	TArray<FTransform> CacheTargetBoneInitTfm;
+	TMap<FName,FTransform> CacheMapBoneSourceTargetTfm;
 	// TMap<FName,FName> CacheIKBoneGoalMap;
 	TMap<FName, const FResolvedBoneChain*> CacheBoneChains;
 	TMap<FName,int32> CacheSourceSkelIndices;
