RelativeIKOp: Allow no-copy codepath when smoothing set to 0. Move smoothing to standard parameters and remove display of testing/debugging ui elements.

diff --git a/Engine/Plugins/Experimental/Animation/RelativeIKOp/Source/RelativeIKOp/Private/RelativeIKOp.cpp b/Engine/Plugins/Experimental/Animation/RelativeIKOp/Source/RelativeIKOp/Private/RelativeIKOp.cpp
--- a/Engine/Plugins/Experimental/Animation/RelativeIKOp/Source/RelativeIKOp/Private/RelativeIKOp.cpp
+++ b/Engine/Plugins/Experimental/Animation/RelativeIKOp/Source/RelativeIKOp/Private/RelativeIKOp.cpp
@@ -103,8 +103,7 @@ void FIKRetargetRelativeIKOp::Run(
 		return;
 	}
 	
-	TArray<FVector3f> SmoothedPoints;
-	ApplyTemporalSmoothing(SmoothedPoints, AnimSeqPlayHead);
+	TConstArrayView<FVector3f> FramePointView = ApplyTemporalSmoothing(AnimSeqPlayHead);
 
 	FDebugRelativeIKDrawInfo LocalDebugInfo;
 	
@@ -178,8 +177,8 @@ void FIKRetargetRelativeIKOp::Run(
 		// Compute source transforms and source-space vertex distance
 		// const FVector3f& LocalPrimaryVert = CachedNotifyInfo->BodyPairsLocalReference[FramePairIdx + 2*PairIdx];
 		// const FVector3f& LocalOtherVert = CachedNotifyInfo->BodyPairsLocalReference[FramePairIdx + 2*PairIdx + 1];
-		const FVector3f& LocalPrimaryVert = SmoothedPoints[2*PairIdx];
-		const FVector3f& LocalOtherVert = SmoothedPoints[2*PairIdx+1];
+		const FVector3f& LocalPrimaryVert = FramePointView[2*PairIdx];
+		const FVector3f& LocalOtherVert = FramePointView[2*PairIdx+1];
 
 		// TODO: Clean this up if we bake in rot invariant space
 		FVector SourceBoneBakePrimaryVert = (SourcePrimaryBodyRot).TransformPosition(FVector(LocalPrimaryVert)); 
@@ -517,42 +516,39 @@ void FIKRetargetRelativeIKOp::UpdateCacheChainInfo(const TArray<FTransform>& Tar
 	CacheChainStartMap.Emplace(BoneName, TargetPose[CacheBoneChains[BoneName]->BoneIndices[0]].GetTranslation());
 }
 
-void FIKRetargetRelativeIKOp::ApplyTemporalSmoothing(TArray<FVector3f>& SmoothedPoints, float Time) const
+TConstArrayView<FVector3f> FIKRetargetRelativeIKOp::ApplyTemporalSmoothing(float Time) const
 {
 	int32 NumBodies = CachedNotifyInfo->BodyPairs.Num();
 	int32 SampleIdx = FMath::Clamp(static_cast<int32>(Time * SampleRate), 0, CachedNotifyInfo->NumSamples-1);
 
-	SmoothedPoints.SetNumUninitialized(NumBodies);
+	static TArray<FVector3f> SmoothedPoints;
 
 	// Testing temporal smoothing of local points
-	const int32 MinSmoothIdx = FMath::Max(SampleIdx - Settings.TestTemporalSmoothingRadius, 0);
-	const int32 MaxSmoothIdx = FMath::Min(SampleIdx + Settings.TestTemporalSmoothingRadius,CachedNotifyInfo->NumSamples-1) ;
+	const int32 MinSmoothIdx = FMath::Max(SampleIdx - Settings.TemporalSmoothingRadius, 0);
+	const int32 MaxSmoothIdx = FMath::Min(SampleIdx + Settings.TemporalSmoothingRadius,CachedNotifyInfo->NumSamples-1) ;
 	
-	if (Settings.TestTemporalSmoothingRadius > 0 && CachedNotifyInfo->NumSamples > 0)
+	if (Settings.TemporalSmoothingRadius <= 0 || CachedNotifyInfo->NumSamples <= 0)
 	{
-		for (int BodyIdx = 0; BodyIdx < CachedNotifyInfo->BodyPairs.Num(); ++BodyIdx)
-		{
-			SmoothedPoints[BodyIdx] = FVector3f::Zero();
-			float SumWeight = 0.0f;
-			for (int32 FrameIdx = MinSmoothIdx; FrameIdx <= MaxSmoothIdx; ++FrameIdx)
-			{
-				// Finite support Gaussian smoothing (could approximate w/ triangle or bake out smoothed data)
-				const float Sigma = 0.25f * static_cast<float>(Settings.TestTemporalSmoothingRadius);
-				const float SmoothDist = static_cast<float>(FrameIdx - SampleIdx);
-				const float SmoothWeight = FMath::Exp(-0.5f*(SmoothDist*SmoothDist) / (Sigma*Sigma));
-				SmoothedPoints[BodyIdx] += SmoothWeight * CachedNotifyInfo->BodyPairsLocalReference[NumBodies * FrameIdx + BodyIdx];
-				SumWeight += SmoothWeight;
-			}
-			SmoothedPoints[BodyIdx] /= SumWeight;
-		}
+		return {CachedNotifyInfo->BodyPairsLocalReference.GetData() + NumBodies * SampleIdx, NumBodies};
 	}
-	else
+	
+	SmoothedPoints.SetNumUninitialized(NumBodies);
+	for (int BodyIdx = 0; BodyIdx < CachedNotifyInfo->BodyPairs.Num(); ++BodyIdx)
 	{
-		for (int BodyIdx = 0; BodyIdx < CachedNotifyInfo->BodyPairs.Num(); ++BodyIdx)
+		SmoothedPoints[BodyIdx] = FVector3f::Zero();
+		float SumWeight = 0.0f;
+		for (int32 FrameIdx = MinSmoothIdx; FrameIdx <= MaxSmoothIdx; ++FrameIdx)
 		{
-			SmoothedPoints[BodyIdx] = CachedNotifyInfo->BodyPairsLocalReference[NumBodies * SampleIdx + BodyIdx];
+			// Finite support Gaussian smoothing (could approximate w/ triangle or bake out smoothed data)
+			const float Sigma = 0.25f * static_cast<float>(Settings.TemporalSmoothingRadius);
+			const float SmoothDist = static_cast<float>(FrameIdx - SampleIdx);
+			const float SmoothWeight = FMath::Exp(-0.5f*(SmoothDist*SmoothDist) / (Sigma*Sigma));
+			SmoothedPoints[BodyIdx] += SmoothWeight * CachedNotifyInfo->BodyPairsLocalReference[NumBodies * FrameIdx + BodyIdx];
+			SumWeight += SmoothWeight;
 		}
+		SmoothedPoints[BodyIdx] /= SumWeight;
 	}
+	return SmoothedPoints;
 }
 
 
@@ -596,20 +592,20 @@ void FIKRetargetRelativeIKOp::ComputeFeasibilityRange(FDoubleInterval& OutFeasib
 
 bool FIKRetargetRelativeIKOp::CalcLineSphereIntersect(FDoubleInterval& OutRange, const FVector& Center, double Radius, const FVector& StartPoint, const FVector& EndPoint)
 {
-	// NOTE: This function assumes two intersection points and returns positive ratio along line from start point
-	// Used for getting the endpoint of feasible range
+	// NOTE: This function assumes two intersection points and returns valid positive ratios along line from start point, false if invalid/no intersection
+	// Used for getting the endpoints of feasible range
 	FVector LineVec = EndPoint - StartPoint;
 	FVector LocPt = StartPoint - Center;
 
 	double a = LineVec.SquaredLength();
 	double b = FVector::DotProduct(LineVec, LocPt);
 	double c = LocPt.SquaredLength() - Radius*Radius;
-	if (b*b - a*c < 0.0)
+	if (b*b - a*c < 0.0 || LineVec.IsNearlyZero())
 	{
 		return false;
 	}
 
-	double B = FMath::Sqrt(b*b - a*c);
+	double B = FMath::Sqrt(b*b - 4.0*a*c);
 	double MaxRange = (-b + B) / a;
 	double MinRange = (-b - B) / a;
 	OutRange = {MinRange, MaxRange};
diff --git a/Engine/Plugins/Experimental/Animation/RelativeIKOp/Source/RelativeIKOp/Public/RelativeIKOp.h b/Engine/Plugins/Experimental/Animation/RelativeIKOp/Source/RelativeIKOp/Public/RelativeIKOp.h
--- a/Engine/Plugins/Experimental/Animation/RelativeIKOp/Source/RelativeIKOp/Public/RelativeIKOp.h
+++ b/Engine/Plugins/Experimental/Animation/RelativeIKOp/Source/RelativeIKOp/Public/RelativeIKOp.h
@@ -122,6 +122,10 @@ struct FIKRetargetRelativeIKOpSettings : public FIKRetargetOpSettingsBase
 	UPROPERTY(EditAnywhere, BlueprintReadWrite, Category = "Parameters", meta=(ClampMin = "0", ClampMax = "1"))
 	double ContributionSumWeight = 1.0;
 
+	// Frames of temporal smoothing of body pair verts (0 is no smoothing)
+	UPROPERTY(EditAnywhere, BlueprintReadWrite, Category = "Parameters", meta=(ClampMin = "0", ClampMax = "60"))
+	int32 TemporalSmoothingRadius = 15;
+
 	// Ignore source scaling when computing relative distance relationships
 	UPROPERTY(EditAnywhere, BlueprintReadWrite, Category = "Parameters")
 	bool bIgnoreSourceScale = true;
@@ -159,24 +163,20 @@ struct FIKRetargetRelativeIKOpSettings : public FIKRetargetOpSettingsBase
 	bool bDebugDrawBodyTransforms = false;
 
 	// Run op and display debug info but DON'T update IK Goals
-	UPROPERTY(EditAnywhere, BlueprintReadWrite, Category = "Debug")
+	UPROPERTY(BlueprintReadWrite, Category = "Debug")
 	bool bDryRun = false;
 
 	// Use distance to pair contact to blend towards direct contact
-	UPROPERTY(EditAnywhere, BlueprintReadWrite, Category = "Test Features")
+	UPROPERTY(BlueprintReadWrite, Category = "Test Features")
 	bool bTestDistContactAlpha = true;
 
 	// Dial back distance weighting based on feasibility length of target bone chain
-	UPROPERTY(EditAnywhere, BlueprintReadWrite, Category = "Test Features")
+	UPROPERTY(BlueprintReadWrite, Category = "Test Features")
 	bool bTestFeasibilityWeight = true;
 
 	// Use target (default) when computing pair alternate representations
-	UPROPERTY(EditAnywhere, BlueprintReadWrite, Category = "Test Features")
+	UPROPERTY(BlueprintReadWrite, Category = "Test Features")
 	bool bTestRetargetScale = true;
-
-	// Frames of temporal smoothing of body pair verts (0 is no smoothing)
-	UPROPERTY(EditAnywhere, BlueprintReadWrite, Category = "Test Features", meta=(ClampMin = "0", ClampMax = "60"))
-	int32 TestTemporalSmoothingRadius = 15;
 };
 
 USTRUCT(BlueprintType, meta = (DisplayName = "Relative IK Goals"))
@@ -232,7 +232,7 @@ private:
 	void UpdateCacheChainInfo(const TArray<FTransform>& TargetPose, FName BoneName);
 	
 	void ComputeFeasibilityRange(FDoubleInterval& OutFeasibleRange, FName BoneName, FVector BoneChainOffset, FVector TargetRelativeToBoneV, FVector TargetRelativeToOpposeBoneV);
-	void ApplyTemporalSmoothing(TArray<FVector3f>& SmoothedPoints, float tim) const;
+	TConstArrayView<FVector3f> ApplyTemporalSmoothing(float Time) const;
 
 	FVector GoalLocBlendCompSpace(const FIKRigGoal* Goal, const FTransform& BoneTfm) const;
 
