RelativeIKOp: Fix for crash when switching from montage -> anim where both contain baked relative body notify data.

diff --git a/Engine/Plugins/Experimental/Animation/RelativeIKOp/Source/RelativeIKOp/Private/RelativeIKOp.cpp b/Engine/Plugins/Experimental/Animation/RelativeIKOp/Source/RelativeIKOp/Private/RelativeIKOp.cpp
--- a/Engine/Plugins/Experimental/Animation/RelativeIKOp/Source/RelativeIKOp/Private/RelativeIKOp.cpp
+++ b/Engine/Plugins/Experimental/Animation/RelativeIKOp/Source/RelativeIKOp/Private/RelativeIKOp.cpp
@@ -959,6 +959,7 @@ void FIKRetargetRelativeIKOp::UpdateRelativeIKNotifyInfoAnimSeq(UAnimInstance* S
 	CacheSourceAnimSequence = Cast<UAnimSequence>(MaxNotifyRecord->SourceAsset);
 
 	// NOTE: Only grabs first bake notify
+	ResetCacheNotifyInfo();
 	for (const FAnimNotifyEvent& NotifyEvent : CacheSourceAnimSequence->Notifies)
 	{
 		URelativeBodyBakeAnimNotify* BakedNotifyInfo = Cast<URelativeBodyBakeAnimNotify>(NotifyEvent.Notify);
@@ -977,6 +978,7 @@ void FIKRetargetRelativeIKOp::ResetCacheNotifyInfo()
 	SegmentEndTime = -1.0f;
 	CachedNotifyInfo = nullptr;
 
+	// TODO: Need to bake info in accessible bipartite graph structure and cache less data!
 	// Clear all cached data on update
 	CacheMapBoneSourceTargetTfm.Reset();
 	CacheSourceBodyEffectVertIdx.Reset();
