RelativeIKOp: Fix possible crash, use per-relik op memory for runtime smoothing.

diff --git a/Engine/Plugins/Experimental/Animation/RelativeIKOp/Source/RelativeIKOp/Private/RelativeIKOp.cpp b/Engine/Plugins/Experimental/Animation/RelativeIKOp/Source/RelativeIKOp/Private/RelativeIKOp.cpp
--- a/Engine/Plugins/Experimental/Animation/RelativeIKOp/Source/RelativeIKOp/Private/RelativeIKOp.cpp
+++ b/Engine/Plugins/Experimental/Animation/RelativeIKOp/Source/RelativeIKOp/Private/RelativeIKOp.cpp
@@ -516,13 +516,11 @@ void FIKRetargetRelativeIKOp::UpdateCacheChainInfo(const TArray<FTransform>& Tar
 	CacheChainStartMap.Emplace(BoneName, TargetPose[CacheBoneChains[BoneName]->BoneIndices[0]].GetTranslation());
 }
 
-TConstArrayView<FVector3f> FIKRetargetRelativeIKOp::ApplyTemporalSmoothing(float Time) const
+TConstArrayView<FVector3f> FIKRetargetRelativeIKOp::ApplyTemporalSmoothing(float Time)
 {
 	int32 NumBodies = CachedNotifyInfo->BodyPairs.Num();
 	int32 SampleIdx = FMath::Clamp(static_cast<int32>(Time * SampleRate), 0, CachedNotifyInfo->NumSamples-1);
 
-	static TArray<FVector3f> SmoothedPoints;
-
 	// Testing temporal smoothing of local points
 	const int32 MinSmoothIdx = FMath::Max(SampleIdx - Settings.TemporalSmoothingRadius, 0);
 	const int32 MaxSmoothIdx = FMath::Min(SampleIdx + Settings.TemporalSmoothingRadius,CachedNotifyInfo->NumSamples-1) ;
diff --git a/Engine/Plugins/Experimental/Animation/RelativeIKOp/Source/RelativeIKOp/Public/RelativeIKOp.h b/Engine/Plugins/Experimental/Animation/RelativeIKOp/Source/RelativeIKOp/Public/RelativeIKOp.h
--- a/Engine/Plugins/Experimental/Animation/RelativeIKOp/Source/RelativeIKOp/Public/RelativeIKOp.h
+++ b/Engine/Plugins/Experimental/Animation/RelativeIKOp/Source/RelativeIKOp/Public/RelativeIKOp.h
@@ -232,7 +232,7 @@ private:
 	void UpdateCacheChainInfo(const TArray<FTransform>& TargetPose, FName BoneName);
 	
 	void ComputeFeasibilityRange(FDoubleInterval& OutFeasibleRange, FName BoneName, FVector BoneChainOffset, FVector TargetRelativeToBoneV, FVector TargetRelativeToOpposeBoneV);
-	TConstArrayView<FVector3f> ApplyTemporalSmoothing(float Time) const;
+	TConstArrayView<FVector3f> ApplyTemporalSmoothing(float Time);
 
 	FVector GoalLocBlendCompSpace(const FIKRigGoal* Goal, const FTransform& BoneTfm) const;
 
@@ -281,6 +281,9 @@ private:
 	TObjectPtr<const UPhysicsAsset> TargetPhysicsAsset;
 	
 	TObjectPtr<const UIKRigDefinition> TargetIKRig;
+	
+	// For runtime smoothing
+	TArray<FVector3f> SmoothedPoints;
 
 	// Cache notify-related data for faster weighted goal computations
 	TObjectPtr<const UAnimSequence> CacheSourceAnimSequence;
