RelativeIKOp: Fix for bug and edge cases in chain feasibility checks.

diff --git a/Engine/Plugins/Experimental/Animation/RelativeIKOp/Source/RelativeIKOp/Private/RelativeIKOp.cpp b/Engine/Plugins/Experimental/Animation/RelativeIKOp/Source/RelativeIKOp/Private/RelativeIKOp.cpp
--- a/Engine/Plugins/Experimental/Animation/RelativeIKOp/Source/RelativeIKOp/Private/RelativeIKOp.cpp
+++ b/Engine/Plugins/Experimental/Animation/RelativeIKOp/Source/RelativeIKOp/Private/RelativeIKOp.cpp
@@ -551,7 +551,6 @@ TConstArrayView<FVector3f> FIKRetargetRelativeIKOp::ApplyTemporalSmoothing(float
 	return SmoothedPoints;
 }
 
-
 void FIKRetargetRelativeIKOp::ComputeFeasibilityRange(FDoubleInterval& OutFeasibleRange, FName BoneName, FVector BoneChainOffset, FVector TargetRelativeToBoneV, FVector TargetRelativeToOpposeBoneV)
 {
 	OutFeasibleRange = {0.0, 1.0};
@@ -576,18 +575,13 @@ void FIKRetargetRelativeIKOp::ComputeFeasibilityRange(FDoubleInterval& OutFeasib
 	{
 		// Whole range is feasible
 		OutFeasibleRange = {0.0,1.0};
+		return;
 	}
-	else if (!CalcLineSphereIntersect(IntersectRange, OffsetChainStart, CheckRadius, TargetRelativeToBoneV, TargetRelativeToOpposeBoneV))
-	{
-		// No feasible range, pick closer of the endpoints
-		double RangeEndPt = (ChainDistStart > ChainDistEnd) ? 1.0 : 0.0;
-		OutFeasibleRange = {RangeEndPt,RangeEndPt};
-	}
-	else
-	{
-		// Use sphere intersection range, clamped on [0,1]
-		OutFeasibleRange = {FMath::Clamp(IntersectRange.Min, 0.0,1.0), FMath::Clamp(IntersectRange.Max, 0.0,1.0)};
-	}
+
+	// If valid intersect: Sphere intersection range (return clamped to line segment) 
+	// No valid intersect: Closest point to sphere (return clamped to line segment)
+	CalcLineSphereIntersect(IntersectRange, OffsetChainStart, CheckRadius, TargetRelativeToBoneV, TargetRelativeToOpposeBoneV);
+	OutFeasibleRange = {FMath::Clamp(IntersectRange.Min, 0.0,1.0), FMath::Clamp(IntersectRange.Max, 0.0,1.0)};
 }
 
 bool FIKRetargetRelativeIKOp::CalcLineSphereIntersect(FDoubleInterval& OutRange, const FVector& Center, double Radius, const FVector& StartPoint, const FVector& EndPoint)
@@ -600,12 +594,21 @@ bool FIKRetargetRelativeIKOp::CalcLineSphereIntersect(FDoubleInterval& OutRange,
 	double a = LineVec.SquaredLength();
 	double b = FVector::DotProduct(LineVec, LocPt);
 	double c = LocPt.SquaredLength() - Radius*Radius;
-	if (b*b - a*c < 0.0 || LineVec.IsNearlyZero())
+	if (LineVec.IsNearlyZero())
 	{
+		// Line is ill-defined, use midpoint
+		OutRange = {0.5,0.5};
 		return false;
 	}
-
-	double B = FMath::Sqrt(b*b - 4.0*a*c);
+	
+	if (b*b - a*c < 0.0)
+	{
+		// No intersection, return closest point on line (clamped to line segment outside this function)
+		OutRange = {-b/a, -b/a};
+		return false;
+	}
+	
+	double B = FMath::Sqrt(b*b - a*c);
 	double MaxRange = (-b + B) / a;
 	double MinRange = (-b - B) / a;
 	OutRange = {MinRange, MaxRange};
