RelativeIKOp: Clean up transform calculations for body pair verts. Also corrected source scale in debug draws for handling changes in CL 45498499.

diff --git a/Engine/Plugins/Experimental/Animation/RelativeIKOp/Source/RelativeIKOp/Private/RelativeIKOp.cpp b/Engine/Plugins/Experimental/Animation/RelativeIKOp/Source/RelativeIKOp/Private/RelativeIKOp.cpp
--- a/Engine/Plugins/Experimental/Animation/RelativeIKOp/Source/RelativeIKOp/Private/RelativeIKOp.cpp
+++ b/Engine/Plugins/Experimental/Animation/RelativeIKOp/Source/RelativeIKOp/Private/RelativeIKOp.cpp
@@ -22,7 +22,6 @@
 
 #define LOCTEXT_NAMESPACE "RelativeIKOp"
 
-
 const UClass* FIKRetargetRelativeIKOpSettings::GetControllerType() const
 {
 	return UIKRetargetRelativeIKController::StaticClass();
@@ -117,10 +116,6 @@ void FIKRetargetRelativeIKOp::Run(
 
 	const double SourceScale = InProcessor.GetSourceScaleFactor();
 	const double InvSourceScale = (Settings.bIgnoreSourceScale) ? 1.0 / InProcessor.GetSourceScaleFactor() : 1.0;
-
-	// NOTE: Uniform scaling so these should compose fine w/ FTransform
-	const FTransform SourceScaleTfm(FQuat::Identity, FVector::Zero(), FVector(SourceScale));
-	const FTransform InvScaleTfm(FQuat::Identity, FVector::Zero(), FVector(InvSourceScale));
 	
 	FIKRigGoalContainer& GoalContainer = InProcessor.GetIKRigGoalContainer();
 	TArray<FIKRigGoal>& IKGoals = GoalContainer.GetGoalArray();
@@ -169,10 +164,8 @@ void FIKRetargetRelativeIKOp::Run(
 		int32 SourceOtherPoseIdx = CacheSourceSkelIndices[SourceOtherBone];
 
 		// ---- TESTING ONLY - Remove these once notify baking is rotation invariant
-		FTransform SourcePrimaryBodyRot = GetBodyRotation(SourcePhysicsAsset, SourcePrimaryBone);
-		FTransform SourceOtherBodyRot = GetBodyRotation(SourcePhysicsAsset, SourceOtherBone);
-		FTransform TargetPrimaryBodyRot = SourcePrimaryBodyRot;
-		FTransform TargetOtherBodyRot = SourceOtherBodyRot;
+		FTransform SourcePrimaryBodyRot = FTransform(GetBodyRotation(SourcePhysicsAsset, SourcePrimaryBone));
+		FTransform SourceOtherBodyRot = FTransform(GetBodyRotation(SourcePhysicsAsset, SourceOtherBone));
 		
 		// Compute source transforms and source-space vertex distance
 		// const FVector3f& LocalPrimaryVert = CachedNotifyInfo->BodyPairsLocalReference[FramePairIdx + 2*PairIdx];
@@ -183,63 +176,55 @@ void FIKRetargetRelativeIKOp::Run(
 		// TODO: Clean this up if we bake in rot invariant space
 		FVector SourceBoneBakePrimaryVert = (SourcePrimaryBodyRot).TransformPosition(FVector(LocalPrimaryVert)); 
 		FVector SourceBoneBakeOtherVert = (SourceOtherBodyRot).TransformPosition(FVector(LocalOtherVert));
-		// NOTE: Need to use FMatrix to apply oriented scale properly
-		FMatrix SourceBodyPrimaryScale = GetBodyOrientedScale(SourcePhysicsAsset, SourcePrimaryBone);
-		FMatrix SourceBodyOtherScale = GetBodyOrientedScale(SourcePhysicsAsset, SourceOtherBone);
-		FVector SourceBodyPrimaryV = SourceBodyPrimaryScale.TransformPosition(SourceBoneBakePrimaryVert);
-		FVector SourceBodyOtherV = SourceBodyOtherScale.TransformPosition(SourceBoneBakeOtherVert);
-		
-		FTransform SourcePrimaryBodyTfm = GetBodyTranslation(SourcePhysicsAsset, SourcePrimaryBone) * SourceScaleTfm;
-		FTransform SourcePrimaryCompTfm = SourcePrimaryBodyTfm * InSourceGlobalPose[SourcePrimaryPoseIdx] * InvScaleTfm;
-		FVector SourcePrimaryCompV = SourcePrimaryCompTfm.TransformPosition(SourceBodyPrimaryV);
 
-		FTransform SourceOtherBodyTfm = GetBodyTranslation(SourcePhysicsAsset, SourceOtherBone) * SourceScaleTfm;
-		FTransform SourceOtherCompTfm = SourceOtherBodyTfm * InSourceGlobalPose[SourceOtherPoseIdx] * InvScaleTfm;
-		FVector SourceOtherCompV = SourceOtherCompTfm.TransformPosition(SourceBodyOtherV);
+		FBodyTransform SourcePrimaryTfm;
+		FBodyTransform SourceOtherTfm;
 
-		double SourceVDist = FVector::Distance(SourcePrimaryCompV, SourceOtherCompV);
+		ComputeSourceBodyTransform(SourcePrimaryTfm, SourcePrimaryBone, InSourceGlobalPose[SourcePrimaryPoseIdx], SourceScale);
+		ComputeSourceBodyTransform(SourceOtherTfm, SourceOtherBone, InSourceGlobalPose[SourceOtherPoseIdx], SourceScale);
+
+		FVector SourcePrimaryCompV = ApplyBodyTransform(SourcePrimaryTfm, SourceBoneBakePrimaryVert);
+		FVector SourceOtherCompV = ApplyBodyTransform(SourceOtherTfm, SourceBoneBakeOtherVert);
+
+		// Rescale components for distance checks if "Ignore sources scale" is on
+		// NOTE: This works because we don't need absolute coordinates (just distance).
+		//       If absolute coords needed then must properly account for scale pivot
+		FVector RescaleSourcePrimaryV = InvSourceScale * SourcePrimaryCompV;
+		FVector RecsaleSourceOtherV = InvSourceScale * SourceOtherCompV;
+		
+		double SourceVDist = FVector::Distance(RescaleSourcePrimaryV, RecsaleSourceOtherV);
 		double LinearWeight = FMath::Clamp(1.0 - (SourceVDist / Settings.DistanceThreshold), 0.0, 1.0);
 		double Weight = GetDistanceWeight(SourceVDist, Settings.DistanceThreshold, Settings.DistanceFade);
 
-		// Compute Source pair verts as represented in alternate pair body local space
-		FVector SourcePrimaryOtherV = SourceBodyOtherScale.InverseTransformPosition(SourceOtherCompTfm.InverseTransformPosition(SourcePrimaryCompV));
-		FVector SourceOtherPrimaryV = SourceBodyPrimaryScale.InverseTransformPosition(SourcePrimaryCompTfm.InverseTransformPosition(SourceOtherCompV));
+		// Compute Source pair verts as represented in alternate pair body primary(_P) or other(_O) local bake space
+		FVector SourcePrimaryV_O = InverseBodyTransform(SourceOtherTfm, SourcePrimaryCompV);
+		FVector SourceOtherV_P = InverseBodyTransform(SourcePrimaryTfm, SourceOtherCompV);
 		
 		//---------------------------------
 		// Get verts retargeted using physics body oriented scale
 		int32 TargetPrimaryPoseIdx = CacheTargetSkelIndices[TargetPrimaryBone];
 		int32 TargetOtherPoseIdx = CacheTargetSkelIndices[TargetOtherBone];
 
-		FTransform TargetPrimaryBoneDelta = GetTargetBoneDelta(TargetPrimaryBone);
-		FTransform TargetOtherBoneDelta = GetTargetBoneDelta(TargetOtherBone);
-		
-		FVector TargetBoneBakePrimaryVert = (TargetPrimaryBodyRot * TargetPrimaryBoneDelta).TransformPosition(FVector(LocalPrimaryVert)); 
-		FVector TargetBoneBakeOtherVert = (TargetOtherBodyRot * TargetOtherBoneDelta).TransformPosition(FVector(LocalOtherVert));
-
-		// NOTE: Need to use FMatrix to apply oriented scale properly
-		FMatrix TargetBodyPrimaryScale = GetBodyOrientedScale(TargetPhysicsAsset, TargetPrimaryBone);
-		FMatrix TargetBodyOtherScale = GetBodyOrientedScale(TargetPhysicsAsset, TargetOtherBone);
-		FVector TargetBodyPrimaryV = TargetBodyPrimaryScale.TransformPosition(TargetBoneBakePrimaryVert);
-		FVector TargetBodyOtherV = TargetBodyOtherScale.TransformPosition(TargetBoneBakeOtherVert);
-		
-		FTransform TargetPrimaryBodyTfm = GetBodyTranslation(TargetPhysicsAsset, TargetPrimaryBone);
-		FTransform TargetOtherBodyTfm = GetBodyTranslation(TargetPhysicsAsset, TargetOtherBone);
+		FBodyTransform TargetPrimaryTfm;
+		FBodyTransform TargetOtherTfm;
 
-		FTransform TargetPrimaryCompTfm = TargetPrimaryBodyTfm * OutTargetGlobalPose[TargetPrimaryPoseIdx];
-		FTransform TargetOtherCompTfm = TargetOtherBodyTfm * OutTargetGlobalPose[TargetOtherPoseIdx];
+		ComputeTargetBodyTransform(TargetPrimaryTfm, TargetPrimaryBone, OutTargetGlobalPose[TargetPrimaryPoseIdx]);
+		ComputeTargetBodyTransform(TargetOtherTfm, TargetOtherBone, OutTargetGlobalPose[TargetOtherPoseIdx]);
 
-		//TODO: Play with/think about offset strategy for these points
-		FMatrix TestPrimaryScale = (Settings.bTestRetargetScale) ? TargetBodyPrimaryScale : SourceBodyPrimaryScale;
-		FMatrix TestOtherScale = (Settings.bTestRetargetScale) ? TargetBodyOtherScale : SourceBodyOtherScale;
-		FVector TargetPrimaryCompV_P = TargetPrimaryCompTfm.TransformPosition(TargetBodyPrimaryV);
-		FVector TargetPrimaryCompV_O = TargetOtherCompTfm.TransformPosition(TestOtherScale.TransformPosition(TargetOtherBoneDelta.TransformPosition(SourcePrimaryOtherV)));
-		FVector TargetPrimaryBoneOffset = OutTargetGlobalPose[TargetPrimaryPoseIdx].GetTranslation() - TargetPrimaryCompV_P;
+		// Compute "primary" body vert represented in component space relative to primary (_P) and other (_O) bodies
+		FVector TargetPrimaryCompV_P = ApplyBodyTransform(TargetPrimaryTfm, SourceBoneBakePrimaryVert);
+		FVector TargetPrimaryCompV_O = ApplyBodyTransform(TargetOtherTfm, SourcePrimaryV_O);
 
-		FVector TargetOtherCompV_O = TargetOtherCompTfm.TransformPosition(TargetBodyOtherV);
-		FVector TargetOtherCompV_P = TargetPrimaryCompTfm.TransformPosition(TestPrimaryScale.TransformPosition(TargetPrimaryBoneDelta.TransformPosition(SourceOtherPrimaryV)));
+		// Compute "other" body vert represented in component space relative to primary (_P) and other (_O) bodies
+		FVector TargetOtherCompV_O = ApplyBodyTransform(TargetOtherTfm, SourceBoneBakeOtherVert);
+		FVector TargetOtherCompV_P = ApplyBodyTransform(TargetPrimaryTfm, SourceOtherV_P);
+
+		// Get offsets from target verts (in their own body rep) to associated bone
+		FVector TargetPrimaryBoneOffset = OutTargetGlobalPose[TargetPrimaryPoseIdx].GetTranslation() - TargetPrimaryCompV_P;
 		FVector TargetOtherBoneOffset = OutTargetGlobalPose[TargetOtherPoseIdx].GetTranslation() - TargetOtherCompV_O;
 
 		// double SpringAlpha = (Settings.bTestDistContactAlpha) ? Weight : RetargetSpringAlpha;
+		// These are the target midpoint positions used for spring constraints in PhysicalIK
 		FVector TargetRelativeToPrimaryV = FMath::Lerp(TargetPrimaryCompV_P,TargetOtherCompV_P, RetargetSpringAlpha);
 		FVector TargetRelativeToOtherV = FMath::Lerp(TargetPrimaryCompV_O,TargetOtherCompV_O, RetargetSpringAlpha);
 
@@ -331,21 +316,19 @@ void FIKRetargetRelativeIKOp::Run(
 		double MaxWeight = 0;
 		for (int Index : EffectIndices)
 		{
-			// if (MaxWeight < FrameBoneWeights[Index])
-			// {
-			// 	MaxWeight = FrameBoneWeights[Index];
-			// 	GoalPoint = FrameBoneTargets[Index];
-			// }
-			GoalPoint += FrameBoneWeights[Index] * FrameBoneTargets[Index];
-			TotalWeight += FrameBoneWeights[Index];
+			if (!FMath::IsNearlyZero(FrameBoneWeights[Index]))
+			{
+				GoalPoint += FrameBoneWeights[Index] * FrameBoneTargets[Index];
+				TotalWeight += FrameBoneWeights[Index];
 
 #if WITH_EDITOR
-			if (Settings.bDebugDraw)
-			{
-				LocalDebugInfo.TargetGoals.Last().PairTargets.Add(FrameBoneTargets[Index]);
-				LocalDebugInfo.PairRetargetInfo.Emplace(LocalTargetPairSpaces[Index]);
-			}
+				if (Settings.bDebugDraw)
+				{
+					LocalDebugInfo.TargetGoals.Last().PairTargets.Add(FrameBoneTargets[Index]);
+					LocalDebugInfo.PairRetargetInfo.Emplace(LocalTargetPairSpaces[Index]);
+				}
 #endif //WITH_EDITOR
+			}
 		}
 
 		
@@ -397,7 +380,9 @@ void FIKRetargetRelativeIKOp::Run(
 			}
 
 			int32 SourcePoseIdx = CacheSourceSkelIndices[SourceBoneName];
-			DebugDrawInfo.SourceBodyInfo.Emplace(SourceBoneName, SourceScaleTfm * InSourceGlobalPose[SourcePoseIdx], FTransform::Identity);
+			// Need to apply source-scale to bodies since it's directly applied to translations only in source pose 
+			FTransform SourceScaleTfm{FQuat::Identity, FVector::ZeroVector, FVector(SourceScale)};
+			DebugDrawInfo.SourceBodyInfo.Emplace(SourceBoneName, SourceScale, SourceScaleTfm*InSourceGlobalPose[SourcePoseIdx], FTransform::Identity);
 		}
 
 		// Cache debug info for target domain bodies
@@ -412,7 +397,7 @@ void FIKRetargetRelativeIKOp::Run(
 
 			FTransform TargetBoneDelta = GetTargetBoneDelta(TargetBoneName);
 			int32 TargetPoseIdx = CacheTargetSkelIndices[TargetBoneName];
-			DebugDrawInfo.TargetBodyInfo.Emplace(TargetBoneName, OutTargetGlobalPose[TargetPoseIdx], TargetBoneDelta);
+			DebugDrawInfo.TargetBodyInfo.Emplace(TargetBoneName, 1.0, OutTargetGlobalPose[TargetPoseIdx], TargetBoneDelta);
 		}
 
 		DebugDrawInfo.SourceTfmInfo.Reset(CacheSourceEffectBones.Num());
@@ -424,17 +409,15 @@ void FIKRetargetRelativeIKOp::Run(
 			}
 
 			int32 SourcePoseIdx = CacheSourceSkelIndices[SourceBoneName];
-
-			// NOTE: Need to use FMatrix to apply oriented scale properly
-			FMatrix SourceBodyRot = GetBodyRotation(SourcePhysicsAsset, SourceBoneName).ToMatrixNoScale();
-			FMatrix SourceBodyScale = GetBodyOrientedScale(SourcePhysicsAsset, SourceBoneName);
-			FMatrix SourceBodyTfm = SourceBodyRot * SourceBodyScale * GetBodyTranslation(SourcePhysicsAsset, SourceBoneName).ToMatrixNoScale();
-			FMatrix SourceCompTfm = SourceBodyTfm * InSourceGlobalPose[SourcePoseIdx].ToMatrixWithScale();
-
-			FVector SourceBodyCenter = InSourceGlobalPose[SourcePoseIdx].TransformPosition(GetBodyTranslation(SourcePhysicsAsset, SourceBoneName).GetLocation());
-			FVector SourceBodyX = SourceCompTfm.TransformPosition(FVector::UnitX());
-			FVector SourceBodyY = SourceCompTfm.TransformPosition(FVector::UnitY());
-			FVector SourceBodyZ = SourceCompTfm.TransformPosition(FVector::UnitZ());
+			
+			FBodyTransform SourceFullTfm;
+			FTransform SourceBodyRot = FTransform(GetBodyRotation(SourcePhysicsAsset, SourceBoneName));
+			ComputeSourceBodyTransform(SourceFullTfm, SourceBoneName, InSourceGlobalPose[SourcePoseIdx], SourceScale);
+
+			FVector SourceBodyCenter = ApplyBodyTransform(SourceFullTfm, FVector::ZeroVector);
+			FVector SourceBodyX = ApplyBodyTransform(SourceFullTfm, SourceBodyRot.GetUnitAxis(EAxis::X));
+			FVector SourceBodyY = ApplyBodyTransform(SourceFullTfm, SourceBodyRot.GetUnitAxis(EAxis::Y));
+			FVector SourceBodyZ = ApplyBodyTransform(SourceFullTfm, SourceBodyRot.GetUnitAxis(EAxis::Z));
 			
 			DebugDrawInfo.SourceTfmInfo.Emplace(SourceBodyCenter, SourceBodyX, SourceBodyY, SourceBodyZ);
 		}
@@ -450,18 +433,14 @@ void FIKRetargetRelativeIKOp::Run(
 			
 			int32 TargetPoseIdx = CacheTargetSkelIndices[TargetBoneName];
 
-			// NOTE: Need to use FMatrix to apply oriented scale properly
-			FMatrix TargetBoneDelta = GetTargetBoneDelta(TargetBoneName).ToMatrixNoScale();
-			
-			FMatrix SourceBodyRot = GetBodyRotation(SourcePhysicsAsset, SourceBoneName).ToMatrixNoScale();
-			FMatrix TargetBodyScale = GetBodyOrientedScale(TargetPhysicsAsset, TargetBoneName);
-			FMatrix TargetBodyTfm = SourceBodyRot * TargetBoneDelta * TargetBodyScale * GetBodyTranslation(TargetPhysicsAsset, TargetBoneName).ToMatrixNoScale();
-			FMatrix TargetCompTfm = TargetBodyTfm * OutTargetGlobalPose[TargetPoseIdx].ToMatrixWithScale();
-
-			FVector TargetBodyCenter = OutTargetGlobalPose[TargetPoseIdx].TransformPosition(GetBodyTranslation(TargetPhysicsAsset, TargetBoneName).GetLocation());
-			FVector TargetBodyX = TargetCompTfm.TransformPosition(FVector::UnitX());
-			FVector TargetBodyY = TargetCompTfm.TransformPosition(FVector::UnitY());
-			FVector TargetBodyZ = TargetCompTfm.TransformPosition(FVector::UnitZ());
+			FBodyTransform TargetFullTfm;
+			FTransform SourceBodyRot = FTransform(GetBodyRotation(SourcePhysicsAsset, SourceBoneName));
+			ComputeTargetBodyTransform(TargetFullTfm, SourceBoneName, OutTargetGlobalPose[TargetPoseIdx]);
+
+			FVector TargetBodyCenter = ApplyBodyTransform(TargetFullTfm, FVector::ZeroVector);
+			FVector TargetBodyX = ApplyBodyTransform(TargetFullTfm, SourceBodyRot.GetUnitAxis(EAxis::X));
+			FVector TargetBodyY = ApplyBodyTransform(TargetFullTfm, SourceBodyRot.GetUnitAxis(EAxis::Y));
+			FVector TargetBodyZ = ApplyBodyTransform(TargetFullTfm, SourceBodyRot.GetUnitAxis(EAxis::Z));
 			
 			DebugDrawInfo.TargetTfmInfo.Emplace(TargetBodyCenter, TargetBodyX, TargetBodyY, TargetBodyZ);
 		}
@@ -469,7 +448,7 @@ void FIKRetargetRelativeIKOp::Run(
 #endif
 }
 
-FName FIKRetargetRelativeIKOp::ApplyBodyMap(FName BodyName)
+FName FIKRetargetRelativeIKOp::ApplyBodyMap(FName BodyName) const
 {
 	if (Settings.BodyMapping.IsEmpty() || !Settings.BodyMapping.Contains(BodyName))
 	{
@@ -493,6 +472,51 @@ const FTransform& FIKRetargetRelativeIKOp::GetTargetBoneDelta(FName TargetBoneNa
 	return CacheMapBoneSourceTargetTfm[TargetBoneName];
 }
 
+void FIKRetargetRelativeIKOp::ComputeSourceBodyTransform(FBodyTransform& OutTransform, FName SourceBoneName, const FTransform& GlobalTfm, double SourceScale) const
+{
+	FQuat BodyRot = GetBodyRotation(SourcePhysicsAsset, SourceBoneName);
+	FVector BodyTrans = GetBodyTranslation(SourcePhysicsAsset, SourceBoneName) * SourceScale;
+	FVector BodyScale = GetBodyOrientedScale(SourcePhysicsAsset, SourceBoneName) * SourceScale;
+
+	OutTransform.BoneToBody = FTransform(BodyRot.Inverse());
+	OutTransform.BodyScale = BodyScale;
+	OutTransform.BodyToGlobal = FTransform(BodyRot, BodyTrans) * GlobalTfm;
+}
+
+void FIKRetargetRelativeIKOp::ComputeTargetBodyTransform(FBodyTransform& OutTransform, FName SourceBoneName, const FTransform& GlobalTfm) const
+{
+	FName TargetBoneName = ApplyBodyMap(SourceBoneName);
+	
+	FQuat BodyRot = GetBodyRotation(TargetPhysicsAsset, TargetBoneName);
+	FVector BodyTrans = GetBodyTranslation(TargetPhysicsAsset, TargetBoneName);
+	// TODO: Probably should always use target body scale
+	FVector BodyScale = (Settings.bTestRetargetScale) ? GetBodyOrientedScale(TargetPhysicsAsset, TargetBoneName)
+													: GetBodyOrientedScale(SourcePhysicsAsset, SourceBoneName);
+	
+	OutTransform.BoneToBody = GetTargetBoneDelta(TargetBoneName) * FTransform(BodyRot.Inverse());
+	OutTransform.BodyScale = BodyScale;
+	OutTransform.BodyToGlobal = FTransform(BodyRot, BodyTrans) * GlobalTfm;
+}
+
+FVector FIKRetargetRelativeIKOp::ApplyBodyTransform(const FBodyTransform& Transform, const FVector& LocalPos)
+{
+	const FTransform& BodyLocalTfm = Transform.BoneToBody;
+	const FVector& LocalScale = Transform.BodyScale;
+	const FTransform& GlobalTfm = Transform.BodyToGlobal;
+	
+	return GlobalTfm.TransformPosition(BodyLocalTfm.TransformPositionNoScale(LocalPos) * LocalScale);
+}
+
+FVector FIKRetargetRelativeIKOp::InverseBodyTransform(const FBodyTransform& Transform, const FVector& GlobalPos)
+{
+	const FTransform& BodyLocalTfm = Transform.BoneToBody;
+	const FVector& LocalScale = Transform.BodyScale;
+	const FTransform& GlobalTfm = Transform.BodyToGlobal;
+
+	return BodyLocalTfm.InverseTransformPositionNoScale(GlobalTfm.InverseTransformPosition(GlobalPos) / LocalScale);
+}
+
+
 void FIKRetargetRelativeIKOp::UpdateCacheChainInfo(const TArray<FTransform>& TargetPose, FName BoneName)
 {
 	if (!CacheBoneChains.Contains(BoneName))
@@ -614,85 +638,82 @@ bool FIKRetargetRelativeIKOp::CalcLineSphereIntersect(FDoubleInterval& OutRange,
 	return true;
 }
 
-FTransform FIKRetargetRelativeIKOp::CalcReferenceShapeScale3D(FKShapeElem* ShapeElem)
+FVector FIKRetargetRelativeIKOp::CalcReferenceShapeScale3D(FKShapeElem* ShapeElem)
 {
-	FTransform ElementScale3D = FTransform::Identity;
+	FVector ElementScale3D = FVector::OneVector;
 	switch (ShapeElem->GetShapeType())
 	{
 	case EAggCollisionShape::Sphere:
 		{
 			FKSphereElem* SphereElem = static_cast<FKSphereElem*>(ShapeElem);
-			ElementScale3D.SetScale3D(FVector(SphereElem->Radius));
+			ElementScale3D = FVector(SphereElem->Radius);
 			break;
 		}
 
 	case EAggCollisionShape::Box:
 		{
 			FKBoxElem* BoxElem = static_cast<FKBoxElem*>(ShapeElem);
-			ElementScale3D.SetScale3D(FVector(BoxElem->X, BoxElem->Y, BoxElem->Z));
+			ElementScale3D = FVector(BoxElem->X, BoxElem->Y, BoxElem->Z);
 			break;
 		}
 
 	case EAggCollisionShape::Sphyl:
 		{
 			FKSphylElem* CapsuleElem = static_cast<FKSphylElem*>(ShapeElem);
-			ElementScale3D.SetScale3D(FVector(CapsuleElem->Radius,CapsuleElem->Radius,CapsuleElem->Length * 0.5f + CapsuleElem->Radius));
+			ElementScale3D = FVector(CapsuleElem->Radius,CapsuleElem->Radius,CapsuleElem->Length * 0.5f + CapsuleElem->Radius);
 			break;
 		}
 
 	case EAggCollisionShape::Convex:
 		{
-			UE_LOG(LogAnimation, Warning, TEXT("Unsupported: Shape is a Convex"));
+			// UE_LOG(LogAnimation, Warning, TEXT("Unsupported: Shape is a Convex"));
 			break;
 		}
 	default:
 		{
-			UE_LOG(LogAnimation, Warning, TEXT("Unknown or unsupported shape type"));
+			// UE_LOG(LogAnimation, Warning, TEXT("Unknown or unsupported shape type"));
 			break;
 		}
 	}
 	
-	ensure(!FMath::IsNearlyZero(ElementScale3D.GetDeterminant()));
+	ensure(!ElementScale3D.IsNearlyZero());
 	
 	return ElementScale3D;
 }
 
-FTransform FIKRetargetRelativeIKOp::GetBodyRotation(const UPhysicsAsset* PhysAsset, FName BoneName)
+FQuat FIKRetargetRelativeIKOp::GetBodyRotation(const UPhysicsAsset* PhysAsset, FName BoneName)
 {
 	FKShapeElem* ShapeElem = FindBodyShape(PhysAsset, BoneName);
 	if (!ShapeElem)
 	{
-		return FTransform::Identity;
+		return FQuat::Identity;
 	}
 	
 	FTransform BodyTransform = ShapeElem->GetTransform();
-	return FTransform(BodyTransform.GetRotation());
+	return BodyTransform.GetRotation();
 }
 
-FTransform FIKRetargetRelativeIKOp::GetBodyTranslation(const UPhysicsAsset* PhysAsset, FName BoneName)
+FVector FIKRetargetRelativeIKOp::GetBodyTranslation(const UPhysicsAsset* PhysAsset, FName BoneName)
 {
 	FKShapeElem* ShapeElem = FindBodyShape(PhysAsset, BoneName);
 	if (!ShapeElem)
 	{
-		return FTransform::Identity;
+		return FVector::ZeroVector;
 	}
 	
 	FTransform BodyTransform = ShapeElem->GetTransform();
-	return FTransform(BodyTransform.GetTranslation());
+	return BodyTransform.GetTranslation();
 }
 
-FMatrix FIKRetargetRelativeIKOp::GetBodyOrientedScale(const UPhysicsAsset* PhysAsset, FName BoneName)
+FVector FIKRetargetRelativeIKOp::GetBodyOrientedScale(const UPhysicsAsset* PhysAsset, FName BoneName)
 {
 	FKShapeElem* ShapeElem = FindBodyShape(PhysAsset, BoneName);
 	if (!ShapeElem)
 	{
-		return FMatrix::Identity;
+		return FVector::OneVector;
 	}
 	
-	FTransform BodyTransform = ShapeElem->GetTransform();
-	FQuat BodyRotation = BodyTransform.GetRotation();
-	FMatrix ElementScale3D = CalcReferenceShapeScale3D(ShapeElem).ToMatrixWithScale();
-	return BodyRotation.Inverse().ToMatrix() * ElementScale3D * BodyRotation.ToMatrix();
+	return CalcReferenceShapeScale3D(ShapeElem);
 }
 
 FVector FIKRetargetRelativeIKOp::GoalLocBlendCompSpace(const FIKRigGoal* Goal, const FTransform& BoneTfm) const
@@ -843,6 +864,7 @@ void FIKRetargetRelativeIKOp::UpdateCacheSkelInfo(const FRetargetSkeleton& InSou
 	CacheSourceBoneInitTfm = InSourceSkeleton.RetargetPoses.GetGlobalRetargetPose();
 	CacheTargetBoneInitTfm = InTargetSkeleton.RetargetPoses.GetGlobalRetargetPose();
 
+	// TODO: Cache less data/directly bake source bone indices and use retarget chains instead of bone mapping
 	CacheSourceSkelIndices.Reset();
 	CacheTargetSkelIndices.Reset();
 }
@@ -992,8 +1014,8 @@ void FIKRetargetRelativeIKOp::UpdateCacheNotifyInfo(URelativeBodyBakeAnimNotify*
 	CachedNotifyInfo = NotifyInfo;
 
 	int32 NumBodies = CachedNotifyInfo->BodyPairs.Num();
-	FrameBoneTargets.SetNum(NumBodies);
-	FrameBoneWeights.SetNumZeroed(NumBodies);
+	FrameBoneTargets.SetNumUninitialized(NumBodies);
+	FrameBoneWeights.SetNumUninitialized(NumBodies);
 
 	SampleRate = FMath::RoundToFloat(1.0f / (CachedNotifyInfo->BodyPairsSampleTime[1] - CachedNotifyInfo->BodyPairsSampleTime[0]));
 
@@ -1091,15 +1113,10 @@ void FIKRetargetRelativeIKOp::DebugDraw(
 	}
 	
 	FScopeLock ScopeLock(&DebugDataMutex);
-
-	const double SourceScale = InSourceTransform.GetScale3D().GetMax();
-	FTransform SourceNoScaleTfm(InSourceTransform.GetRotation(), InSourceTransform.GetTranslation());
-	const FTransform& SourceScaleCheckedTfm = (Settings.bIgnoreSourceScale) ? InSourceTransform : SourceNoScaleTfm;
-
 	if (Settings.bDebugDrawBodyPairs)
 	{
 		// Draw Source body pairs
-		DebugDrawBodyPairs(InPDI, SourceScaleCheckedTfm, SourceScale, DebugDrawInfo.SourcePairVerts);
+		DebugDrawBodyPairs(InPDI, InSourceTransform, 1.0, DebugDrawInfo.SourcePairVerts);
 
 		// Draw target body pairs (Possibly with "full" retarget space display)
 		DebugDrawBodyPairs(InPDI, InComponentTransform, InComponentScale, DebugDrawInfo.TargetPairVerts);
@@ -1109,7 +1126,7 @@ void FIKRetargetRelativeIKOp::DebugDraw(
 	{
 		if (Settings.SourcePhysicsAssetOverride)
 		{
-			DebugDrawBodies(InPDI, SourceNoScaleTfm, SourceScale, Settings.SourcePhysicsAssetOverride, DebugDrawInfo.SourceBodyInfo);
+			DebugDrawBodies(InPDI, InSourceTransform, 1.0, Settings.SourcePhysicsAssetOverride, DebugDrawInfo.SourceBodyInfo);
 		}
 		if (Settings.TargetPhysicsAssetOverride)
 		{
@@ -1121,7 +1138,7 @@ void FIKRetargetRelativeIKOp::DebugDraw(
 	{
 		if (Settings.SourcePhysicsAssetOverride)
 		{
-			DebugDrawBodyCoords(InPDI, SourceNoScaleTfm, SourceScale, Settings.SourcePhysicsAssetOverride, DebugDrawInfo.SourceTfmInfo);
+			DebugDrawBodyCoords(InPDI, InSourceTransform, 1.0, Settings.SourcePhysicsAssetOverride, DebugDrawInfo.SourceTfmInfo);
 		}
 		if (Settings.TargetPhysicsAssetOverride)
 		{
@@ -1243,7 +1260,7 @@ void FIKRetargetRelativeIKOp::DebugDrawBodies(FPrimitiveDrawInterface* InPDI, co
 		FTransform CompTfm = BodyInfo.BoneTfm * BaseTransform;
 
 		// Draw physics body
-		DebugDrawPhysBody(InPDI, CompTfm, Scale, ShapeElem, FLinearColor::Yellow);
+		DebugDrawPhysBody(InPDI, CompTfm, Scale * BodyInfo.Scale, ShapeElem, FLinearColor::Yellow);
 	}
 }
 
@@ -1320,7 +1337,7 @@ void FIKRetargetRelativeIKOp::ResetDebugInfo()
 	DebugDrawInfo.SourceBodyInfo.Reset();
 	DebugDrawInfo.TargetBodyInfo.Reset();
 
-	DebugDrawInfo.SourceBodyInfo.Reset();
+	DebugDrawInfo.SourceTfmInfo.Reset();
 	DebugDrawInfo.TargetTfmInfo.Reset();
 	
 	DebugDrawInfo.PairRetargetInfo.Reset();
diff --git a/Engine/Plugins/Experimental/Animation/RelativeIKOp/Source/RelativeIKOp/Public/RelativeIKOp.h b/Engine/Plugins/Experimental/Animation/RelativeIKOp/Source/RelativeIKOp/Public/RelativeIKOp.h
--- a/Engine/Plugins/Experimental/Animation/RelativeIKOp/Source/RelativeIKOp/Public/RelativeIKOp.h
+++ b/Engine/Plugins/Experimental/Animation/RelativeIKOp/Source/RelativeIKOp/Public/RelativeIKOp.h
@@ -52,6 +52,7 @@ struct FDebugDrawBodyPair
 struct FDebugDrawBodyInfo
 {
 	FName BodyName;
+	double Scale;
 	FTransform BoneTfm;
 	FTransform RetargetTfm = FTransform::Identity;
 };
@@ -225,11 +226,27 @@ private:
 	void UpdateRelativeIKNotifyInfoAnimSeq(UAnimInstance* SourceAnimInstance);
 	void ResetCacheNotifyInfo();
 	void UpdateCacheNotifyInfo(URelativeBodyBakeAnimNotify* NotifyInfo);
+	void UpdateCacheChainInfo(const TArray<FTransform>& TargetPose, FName BoneName);
 
-	FName ApplyBodyMap(FName BodyName);
+	FName ApplyBodyMap(FName BodyName) const;
 	bool HasBoneDelta(FName TargetBoneName) const;
 	const FTransform& GetTargetBoneDelta(FName TargetBoneName) const;
-	void UpdateCacheChainInfo(const TArray<FTransform>& TargetPose, FName BoneName);
+
+	// Internal structure for holding necessary baked->component transform data
+	// Baked data is in Source Bone orientation but at body origin
+	struct FBodyTransform
+	{
+		// TODO: Test this with FMatrix computations and see what's better/faster
+		FTransform BoneToBody;
+		FVector BodyScale;
+		FTransform BodyToGlobal;
+	};
+	
+	void ComputeSourceBodyTransform(FBodyTransform& OutTransform, FName SourceBoneName, const FTransform& GlobalTfm, double SourceScale) const;
+	void ComputeTargetBodyTransform(FBodyTransform& OutTransform, FName SourceBoneName, const FTransform& GlobalTfm) const;
+
+	static FVector ApplyBodyTransform(const FBodyTransform& Transform, const FVector& LocalPos);
+	static FVector InverseBodyTransform(const FBodyTransform& Transform, const FVector& GlobalPos);
 	
 	void ComputeFeasibilityRange(FDoubleInterval& OutFeasibleRange, FName BoneName, FVector BoneChainOffset, FVector TargetRelativeToBoneV, FVector TargetRelativeToOpposeBoneV);
 	TConstArrayView<FVector3f> ApplyTemporalSmoothing(float Time);
@@ -238,13 +255,13 @@ private:
 
 	static bool CalcLineSphereIntersect(FDoubleInterval& OutRange, const FVector& Center, double Radius, const FVector& StartPoint, const FVector& EndPoint);
 
-	static FTransform CalcReferenceShapeScale3D(FKShapeElem* ShapeElem);
+	static FVector CalcReferenceShapeScale3D(FKShapeElem* ShapeElem);
 	// Get local body rotation relative to parent bone
-	static FTransform GetBodyRotation(const UPhysicsAsset* PhysAsset, FName BoneName);
+	static FQuat GetBodyRotation(const UPhysicsAsset* PhysAsset, FName BoneName);
 	// Get body origin relative to parent bone
-	static FTransform GetBodyTranslation(const UPhysicsAsset* PhysAsset, FName BoneName);
-	// Get body oriented (non-axis aligned scaling) relative to parent bone
-	static FMatrix GetBodyOrientedScale(const UPhysicsAsset* PhysAsset, FName BoneName);
+	static FVector GetBodyTranslation(const UPhysicsAsset* PhysAsset, FName BoneName);
+	// Get body oriented scale vector
+	static FVector GetBodyOrientedScale(const UPhysicsAsset* PhysAsset, FName BoneName);
 	static FKShapeElem* FindBodyShape(const UPhysicsAsset* PhysAsset, FName BoneName);
 	
 	static double GetDistanceWeight(double Distance, double DistThreshold, double ScalarFade = 200.0);
