BodyIntersectIKOp: Updated goal settings to allow sphere scaling and local offset. Goal intersections no longer require physics body attached to goal bones.

diff --git a/Engine/Plugins/Experimental/Animation/RelativeIKOp/Source/BodyIntersectIKOp/Private/BodyIntersectIKOp.cpp b/Engine/Plugins/Experimental/Animation/RelativeIKOp/Source/BodyIntersectIKOp/Private/BodyIntersectIKOp.cpp
--- a/Engine/Plugins/Experimental/Animation/RelativeIKOp/Source/BodyIntersectIKOp/Private/BodyIntersectIKOp.cpp
+++ b/Engine/Plugins/Experimental/Animation/RelativeIKOp/Source/BodyIntersectIKOp/Private/BodyIntersectIKOp.cpp
@@ -68,15 +68,14 @@ bool FIKRetargetBodyIntersectIKOp::Initialize(
 	UpdateCacheSkelInfo(InTargetSkeleton);
 	
 	const FIKRigGoalContainer& GoalContainer = InProcessor.GetIKRigGoalContainer();
-	for (const FName& GoalName : Settings.IntersectGoals)
+	for (const FIKGoalIntersectSphereSettings& GoalSettings : Settings.IntersectGoalSettings)
 	{
-		const FIKRigGoal* IKGoal = GoalContainer.FindGoalByName(GoalName);
+		const FIKRigGoal* IKGoal = GoalContainer.FindGoalByName(GoalSettings.Goal);
 		if (!IKGoal)
 		{
 			continue;
 		}
-
-		// TODO: Assumes goal bone has a phys body associated with it
+		
 		UpdateTargetBoneMap(IKGoal->BoneName);
 	}
 	
@@ -109,9 +108,9 @@ void FIKRetargetBodyIntersectIKOp::Run(
 #endif //WITH_EDITOR
 
 	FIKRigGoalContainer& GoalContainer = InProcessor.GetIKRigGoalContainer();
-	for (const FName& GoalName : Settings.IntersectGoals)
+	for (const FIKGoalIntersectSphereSettings& GoalSettings : Settings.IntersectGoalSettings)
 	{
-		FIKRigGoal* IKGoal = GoalContainer.FindGoalByName(GoalName);
+		FIKRigGoal* IKGoal = GoalContainer.FindGoalByName(GoalSettings.Goal);
 		if (!IKGoal || !IKGoal->bEnabled)
 		{
 			continue;
@@ -133,14 +132,15 @@ void FIKRetargetBodyIntersectIKOp::Run(
 		const int32 TargetBoneIdx = CacheTargetSkelIndices[TargetBoneName];
 		const FTransform TargetBoneTfm = OutTargetGlobalPose[TargetBoneIdx];
 		
-		double GoalRadius = CalcShapeSmallRadius(GoalShape);
+		double GoalRadius = CalcShapeSmallRadius(GoalShape) * GoalSettings.SphereScale;
 		if (FMath::IsNearlyZero(GoalRadius))
 		{
 			continue;
 		}
 
-		// TODO: Handle actual capsule offsets w/ a lever arm for rotation?
+		// TODO: Handle actual rotation w/ a lever arm and second sphere
 		FVector CompGoalUpdate = GoalLocBlendCompSpace(IKGoal, TargetBoneTfm);
+		FVector SphereOffsetGoal = CompGoalUpdate + GoalSettings.SphereOffset;
 		for (FName BodyName : Settings.IntersectBodies)
 		{
 			if (TargetBoneName == BodyName)
@@ -162,17 +162,18 @@ void FIKRetargetBodyIntersectIKOp::Run(
 			}
 
 			FVector DeltaDir = FVector::ZeroVector;
-			double Dist = CalcIntersectionDelta(CompGoalUpdate, GoalRadius, BodyTfm, ShapeElem, DeltaDir);
+			double Dist = CalcIntersectionDelta(SphereOffsetGoal, GoalRadius, BodyTfm, ShapeElem, DeltaDir);
 			if (Dist > 0.0)
 			{
-				CompGoalUpdate += Dist * DeltaDir;
+				SphereOffsetGoal += Dist * DeltaDir;
 			}
 		}
-		SetGoalPosFromCompSpace(IKGoal, TargetBoneTfm, CompGoalUpdate);
+		
+		SetGoalPosFromCompSpace(IKGoal, TargetBoneTfm, SphereOffsetGoal-GoalSettings.SphereOffset);
 #if WITH_EDITOR
 		if (Settings.bDebugDraw)
 		{
-			LocalDebugInfo.TestSpheres.Emplace(CompGoalUpdate, GoalRadius);
+			LocalDebugInfo.TestSpheres.Emplace(SphereOffsetGoal, GoalRadius);
 		}
 #endif //WITH_EDITOR
 	}
@@ -184,11 +185,10 @@ void FIKRetargetBodyIntersectIKOp::Run(
 		
 		DebugDrawInfo.TestSpheres = LocalDebugInfo.TestSpheres;
 
-		// Cache debug info for target domain bodies
+		// Cache debug info for target intersect bodies
 		DebugDrawInfo.TargetIntersectTfms.Reset(Settings.IntersectBodies.Num());
 		for (FName TargetBoneName : Settings.IntersectBodies)
 		{
-			// TODO: Only show domain bodies? 
 			if (!CacheTargetSkelIndices.Contains(TargetBoneName))
 			{
 				continue;
@@ -307,10 +307,20 @@ double FIKRetargetBodyIntersectIKOp::CalcShapeSmallRadius(FKShapeElem* ShapeElem
 			FKSphylElem* CapsuleElem = static_cast<FKSphylElem*>(ShapeElem);
 			return CapsuleElem->Radius;
 		}
+	case EAggCollisionShape::Box:
+		{
+			FKBoxElem* BoxElem = static_cast<FKBoxElem*>(ShapeElem);
+			return 0.5 * FMath::Min(FMath::Min(BoxElem->X, BoxElem->Y), BoxElem->Z);
+		}
+	// case EAggCollisionShape::Convex:
+	// 	{
+	// 		FKConvexElem* ConvElem = static_cast<FKConvexElem*>(ShapeElem);
+	// 		return 0.5 * ConvElem->ElemBox.GetSize().GetMin();
+	// 	}
 
 	default:
 		{
-			return 0.0;
+			return 1.0;
 		}
 	}
 }
@@ -352,7 +362,6 @@ FKShapeElem* FIKRetargetBodyIntersectIKOp::FindBodyShape(const UPhysicsAsset* Ph
 void FIKRetargetBodyIntersectIKOp::UpdateTargetBoneMap(FName TargetBoneName)
 {
 	int32 TargetIdx = TargetBoneNames.Find(TargetBoneName);
-
 	if (TargetIdx == INDEX_NONE)
 	{
 		return;
diff --git a/Engine/Plugins/Experimental/Animation/RelativeIKOp/Source/BodyIntersectIKOp/Public/BodyIntersectIKOp.h b/Engine/Plugins/Experimental/Animation/RelativeIKOp/Source/BodyIntersectIKOp/Public/BodyIntersectIKOp.h
--- a/Engine/Plugins/Experimental/Animation/RelativeIKOp/Source/BodyIntersectIKOp/Public/BodyIntersectIKOp.h
+++ b/Engine/Plugins/Experimental/Animation/RelativeIKOp/Source/BodyIntersectIKOp/Public/BodyIntersectIKOp.h
@@ -22,6 +22,22 @@ struct FDebugBodyIntersectDrawInfo
 	TArray<TTuple<FVector,double>> TestSpheres;
 };
 
+USTRUCT(BlueprintType, meta = (DisplayName = "Goal Intersect Settings"))
+struct FIKGoalIntersectSphereSettings
+{
+	GENERATED_BODY()
+
+	// The IK Goal Name (See Run IK Rig Op) to intersect
+	UPROPERTY(EditAnywhere, BlueprintReadWrite, Category = "Goal Intersect Settings")
+	FName Goal = NAME_None;
+	// Sphere radius scale
+	UPROPERTY(EditAnywhere, BlueprintReadWrite, Category = "Goal Intersect Settings", meta = (ClampMin="0.0"))
+	double SphereScale = 1.0;
+	// Local offset of sphere from goal bone
+	UPROPERTY(EditAnywhere, BlueprintReadWrite, Category = "Goal Intersect Settings")
+	FVector SphereOffset = FVector::ZeroVector;
+};
+
 USTRUCT(BlueprintType, meta = (DisplayName = "Body Intersect IK Settings"))
 struct FIKRetargetBodyIntersectIKOpSettings : public FIKRetargetOpSettingsBase
 {
@@ -35,11 +51,12 @@ struct FIKRetargetBodyIntersectIKOpSettings : public FIKRetargetOpSettingsBase
 	UPROPERTY(EditAnywhere, BlueprintReadWrite, Category = "Physics", meta=(ReinitializeOnEdit))
 	TObjectPtr<UPhysicsAsset> TargetPhysicsAssetOverride;
 
-	// IK Goals to check sphere intersections (Goal effector bone should have a sphere or capsule body!)
+	// IK Goals to sphere intersect against physics asset bodies 
 	UPROPERTY(EditAnywhere, BlueprintReadWrite, Category = "Intersection", meta=(ReinitializeOnEdit))
-	TArray<FName> IntersectGoals;
+	TArray<FIKGoalIntersectSphereSettings> IntersectGoalSettings;
 
-	// Physics bodies to do trivial intersection against
+	// Bone names with attached physics bodies to do trivial intersection against
+	// **NOTE: Only Sphere and Capsule bodies currently supported**
 	UPROPERTY(EditAnywhere, BlueprintReadWrite, Category = "Intersection", meta=(ReinitializeOnEdit))
 	TArray<FName> IntersectBodies;
 };
