BodyIntersectIKOp: Fix sphere offsets to actually be in local goal bone space.

diff --git a/Engine/Plugins/Experimental/Animation/RelativeIKOp/Source/BodyIntersectIKOp/Private/BodyIntersectIKOp.cpp b/Engine/Plugins/Experimental/Animation/RelativeIKOp/Source/BodyIntersectIKOp/Private/BodyIntersectIKOp.cpp
--- a/Engine/Plugins/Experimental/Animation/RelativeIKOp/Source/BodyIntersectIKOp/Private/BodyIntersectIKOp.cpp
+++ b/Engine/Plugins/Experimental/Animation/RelativeIKOp/Source/BodyIntersectIKOp/Private/BodyIntersectIKOp.cpp
@@ -140,7 +140,8 @@ void FIKRetargetBodyIntersectIKOp::Run(
 
 		// TODO: Handle actual rotation w/ a lever arm and second sphere
 		FVector CompGoalUpdate = GoalLocBlendCompSpace(IKGoal, TargetBoneTfm);
-		FVector SphereOffsetGoal = CompGoalUpdate + GoalSettings.SphereOffset;
+		FVector CompSphereOffset = TargetBoneTfm.TransformVector(GoalSettings.SphereOffset);
+		FVector SphereOffsetGoal = CompGoalUpdate + CompSphereOffset;
 		for (FName BodyName : Settings.IntersectBodies)
 		{
 			if (TargetBoneName == BodyName)
@@ -169,7 +170,7 @@ void FIKRetargetBodyIntersectIKOp::Run(
 			}
 		}
 		
-		SetGoalPosFromCompSpace(IKGoal, TargetBoneTfm, SphereOffsetGoal-GoalSettings.SphereOffset);
+		SetGoalPosFromCompSpace(IKGoal, TargetBoneTfm, SphereOffsetGoal-CompSphereOffset);
 #if WITH_EDITOR
 		if (Settings.bDebugDraw)
 		{
