BodyIntersectIKOp: Removed check for matching body on goal bones as physics body should no longer be required.

diff --git a/Engine/Plugins/Experimental/Animation/RelativeIKOp/Source/BodyIntersectIKOp/Private/BodyIntersectIKOp.cpp b/Engine/Plugins/Experimental/Animation/RelativeIKOp/Source/BodyIntersectIKOp/Private/BodyIntersectIKOp.cpp
--- a/Engine/Plugins/Experimental/Animation/RelativeIKOp/Source/BodyIntersectIKOp/Private/BodyIntersectIKOp.cpp
+++ b/Engine/Plugins/Experimental/Animation/RelativeIKOp/Source/BodyIntersectIKOp/Private/BodyIntersectIKOp.cpp
@@ -122,16 +122,11 @@ void FIKRetargetBodyIntersectIKOp::Run(
 			// UE_LOG(LogAnimation, Warning, TEXT("---- Could not find cache bone: %s"), *TargetBoneName.ToString());
 			continue;
 		}
-		
-		FKShapeElem* GoalShape = FindBodyShape(PhysicsAsset, IKGoal->BoneName);
-		if (!GoalShape)
-		{
-			continue;
-		}
 
 		const int32 TargetBoneIdx = CacheTargetSkelIndices[TargetBoneName];
 		const FTransform TargetBoneTfm = OutTargetGlobalPose[TargetBoneIdx];
-		
+
+		FKShapeElem* GoalShape = FindBodyShape(PhysicsAsset, IKGoal->BoneName);
 		double GoalRadius = CalcShapeSmallRadius(GoalShape) * GoalSettings.SphereScale;
 		if (FMath::IsNearlyZero(GoalRadius))
 		{
@@ -295,6 +290,13 @@ double FIKRetargetBodyIntersectIKOp::CalcIntersectionDelta(const FVector& Target
 
 double FIKRetargetBodyIntersectIKOp::CalcShapeSmallRadius(FKShapeElem* ShapeElem)
 {
+	// Return 1.0 as default so scaling can set "size" of sphere
+	const double DefaultRadius = 1.0;
+	if (!ShapeElem)
+	{
+		return DefaultRadius;
+	}
+	
 	switch (ShapeElem->GetShapeType())
 	{
 	case EAggCollisionShape::Sphere:
@@ -321,7 +323,7 @@ double FIKRetargetBodyIntersectIKOp::CalcShapeSmallRadius(FKShapeElem* ShapeElem
 
 	default:
 		{
-			return 1.0;
+			return DefaultRadius;
 		}
 	}
 }
