RelativeBodyAnimUtils: Add blueprint helpers for bulk updating/applying relative body modifiers (baking relative ik data).

diff --git a/Engine/Plugins/Experimental/Animation/RelativeIKOp/Source/RelativeBodyAnimUtils/Private/RelativeBodyAnimModifier.cpp b/Engine/Plugins/Experimental/Animation/RelativeIKOp/Source/RelativeBodyAnimUtils/Private/RelativeBodyAnimModifier.cpp
--- a/Engine/Plugins/Experimental/Animation/RelativeIKOp/Source/RelativeBodyAnimUtils/Private/RelativeBodyAnimModifier.cpp
+++ b/Engine/Plugins/Experimental/Animation/RelativeIKOp/Source/RelativeBodyAnimUtils/Private/RelativeBodyAnimModifier.cpp
@@ -3,6 +3,7 @@
 #include "RelativeBodyAnimModifier.h"
 
 #include "AnimPose.h"
+#include "AnimationModifiersAssetUserData.h"
 #include "MeshDescription.h"
 #include "RelativeBodyAnimNotifies.h"
 #include "RelativeBodyUtils.h"
@@ -18,6 +19,68 @@
 
 #define LOCTEXT_NAMESPACE "RelativeBodyAnimModifier"
 
+URelativeBodyAnimModifier* FindRelativeBodyModifier(UAnimationModifiersAssetUserData* AssetUserData)
+{
+	const TArray<UAnimationModifier*>& ModifierInstances = AssetUserData->GetAnimationModifierInstances();
+	for (UAnimationModifier* Modifier : ModifierInstances)
+	{
+		if (URelativeBodyAnimModifier* CheckInst = Cast<URelativeBodyAnimModifier>(Modifier))
+		{
+			return CheckInst;
+		}
+	}
+
+	return nullptr;
+}
+
+void URelativeBodyAnimBlueprintFunctions::BulkUpdateRelativeBodyModifiers(TArray<UAnimSequenceBase*>& AnimationList, const FRelativeBodyAnimModifierOptions& Options)
+{
+	FScopedSlowTask ApplyModifiersTask(static_cast<float>(AnimationList.Num()), LOCTEXT("BulkUpdateProgress", "Updating and applying modifier to animation list"));
+	ApplyModifiersTask.MakeDialog();
+	for (UAnimSequenceBase* Anim : AnimationList)
+	{
+		ApplyModifiersTask.EnterProgressFrame(1.0f, FText::FromString(Anim->GetName()));
+		UAnimationModifiersAssetUserData* AssetUserData = Anim->GetAssetUserData<UAnimationModifiersAssetUserData>();
+		if (!AssetUserData)
+		{
+			UE_LOG(LogAnimation, Warning, TEXT("Unable to find modifier user data for anim: %s"), *Anim->GetName());
+			continue;
+		}
+
+		TObjectPtr<URelativeBodyAnimModifier> UpdateModifierInst = FindRelativeBodyModifier(AssetUserData);
+		if (!UpdateModifierInst)
+		{
+			if (!AssetUserData->AddAnimationModifierOfClass(Anim, URelativeBodyAnimModifier::StaticClass()))
+			{
+				UE_LOG(LogAnimation, Warning, TEXT("Unable to create new modifier for anim: %s"), *Anim->GetName());
+				continue;
+			}
+
+			UpdateModifierInst = FindRelativeBodyModifier(AssetUserData);
+		}
+
+		if (!UpdateModifierInst)
+		{
+			continue;
+		}
+
+		UpdateModifierInst->SampleRate = Options.SampleRate;
+		UpdateModifierInst->ContactThreshold = Options.ContactThreshold;
+		UpdateModifierInst->SkeletalMeshAsset = Options.SkeletalMeshAsset;
+		UpdateModifierInst->PhysicsAssetOverride = Options.PhysicsAssetOverride;
+		UpdateModifierInst->DomainBodyNames = Options.DomainBodyNames;
+		UpdateModifierInst->ContactBodyNames = Options.ContactBodyNames;
+		UpdateModifierInst->NotifyClass = Options.NotifyClass;
+
+		UAnimSequence* AnimSeq = Cast<UAnimSequence>(Anim);
+		if (AnimSeq)
+		{
+			UpdateModifierInst->ApplyToAnimationSequence(AnimSeq);
+		}
+	}
+}
+
+
 void URelativeBodyAnimModifier::OnApply_Implementation(UAnimSequence* InAnimation)
 {
 	// TODO: Make log category for relative body utils?
diff --git a/Engine/Plugins/Experimental/Animation/RelativeIKOp/Source/RelativeBodyAnimUtils/Public/RelativeBodyAnimModifier.h b/Engine/Plugins/Experimental/Animation/RelativeIKOp/Source/RelativeBodyAnimUtils/Public/RelativeBodyAnimModifier.h
--- a/Engine/Plugins/Experimental/Animation/RelativeIKOp/Source/RelativeBodyAnimUtils/Public/RelativeBodyAnimModifier.h
+++ b/Engine/Plugins/Experimental/Animation/RelativeIKOp/Source/RelativeBodyAnimUtils/Public/RelativeBodyAnimModifier.h
@@ -4,6 +4,7 @@
 
 #include "AnimationModifier.h"
 #include "MeshDescription.h"
+#include "Kismet/BlueprintFunctionLibrary.h"
 #include "UObject/ObjectPtr.h"
 
 #include "RelativeBodyAnimModifier.generated.h"
@@ -38,10 +39,55 @@ struct FRelativeBodySourceData
 	TArray<int32> SourceBoneIndicesToBodyIndices;
 };
 
+USTRUCT(BlueprintType)
+struct RELATIVEBODYANIMUTILS_API FRelativeBodyAnimModifierOptions
+{
+	GENERATED_BODY()
+
+	/** Rate used to sample the animation */
+	UPROPERTY(EditAnywhere, BlueprintReadWrite, Category = "Settings", meta = (Units = "Hz", UIMin = 1))
+	int SampleRate = 30;
+
+	/** Threshold for determining if a bone pair can be considered to be having contact */
+	UPROPERTY(EditAnywhere, BlueprintReadWrite, Category = "Settings")
+	float ContactThreshold = 150.0f;
+
+	UPROPERTY(EditAnywhere, BlueprintReadWrite, Category = "Settings", meta = (TitleProperty = "SkeletalMeshAsset"))
+	TObjectPtr<USkeletalMesh> SkeletalMeshAsset;
+
+	UPROPERTY(EditAnywhere, BlueprintReadWrite, Category = "Settings", meta = (TitleProperty = "PhysicsAssetOverride"))
+	TObjectPtr<UPhysicsAsset> PhysicsAssetOverride;
+
+	/** Bodies to be checked against contact bodies */
+	UPROPERTY(EditAnywhere, BlueprintReadWrite, Category = "Settings", meta = (TitleProperty = "DomainBodyNames"))
+	TArray<FName> DomainBodyNames;
+	
+	UPROPERTY(EditAnywhere, BlueprintReadWrite, Category = "Settings", meta = (TitleProperty = "ContactBodyNames"))
+	TArray<FName> ContactBodyNames;
+
+	// Relative Body Notify subclass to create
+	UPROPERTY(EditAnywhere, BlueprintReadWrite, Category = "Settings", meta = (TitleProperty = "PhysicalIKRetargeterNotifyClass"))
+	TSubclassOf<URelativeBodyAnimNotifyBase> NotifyClass = nullptr;
+};
+
+/**
+ * BPFL helper for setting up anim modifier 
+ */
+UCLASS(Transient, meta=(ScriptName="RelativeBodyAnimModifyHelper"))
+class RELATIVEBODYANIMUTILS_API URelativeBodyAnimBlueprintFunctions : public UBlueprintFunctionLibrary
+{
+	GENERATED_BODY()
+public:
+	
+	UFUNCTION(BlueprintCallable, Category = "Anim Modifier")
+	static void BulkUpdateRelativeBodyModifiers(UPARAM(ref) TArray<UAnimSequenceBase*>& AnimationList, const FRelativeBodyAnimModifierOptions& Options);
+};
+
+
 /**
  * Animation modifier for baking relative body relationships into anim notify
  */
-UCLASS(meta = (IsBlueprintBase = true))
+UCLASS(BlueprintType, meta = (IsBlueprintBase = true))
 class RELATIVEBODYANIMUTILS_API URelativeBodyAnimModifier : public UAnimationModifier
 {
 	GENERATED_BODY()
@@ -54,28 +100,28 @@ public:
 	
 public:
 	/** Rate used to sample the animation */
-	UPROPERTY(EditAnywhere, BlueprintReadOnly, Category = "Settings", meta = (Units = "Hz", UIMin = 1))
+	UPROPERTY(EditAnywhere, BlueprintReadWrite, Category = "Settings", meta = (Units = "Hz", UIMin = 1))
 	int SampleRate = 30;
 
 	/** Threshold for determining if a bone pair can be considered to be having contact */
-	UPROPERTY(EditAnywhere, BlueprintReadOnly, Category = "Settings")
+	UPROPERTY(EditAnywhere, BlueprintReadWrite, Category = "Settings")
 	float ContactThreshold = 150.0f;
 
-	UPROPERTY(EditAnywhere, BlueprintReadOnly, Category = "Settings", meta = (TitleProperty = "SkeletalMeshAsset"))
+	UPROPERTY(EditAnywhere, BlueprintReadWrite, Category = "Settings", meta = (TitleProperty = "SkeletalMeshAsset"))
 	TObjectPtr<USkeletalMesh> SkeletalMeshAsset;
 
-	UPROPERTY(EditAnywhere, BlueprintReadOnly, Category = "Settings", meta = (TitleProperty = "PhysicsAssetOverride"))
+	UPROPERTY(EditAnywhere, BlueprintReadWrite, Category = "Settings", meta = (TitleProperty = "PhysicsAssetOverride"))
 	TObjectPtr<UPhysicsAsset> PhysicsAssetOverride;
 
 	/** Bodies to be checked against contact bodies */
-	UPROPERTY(EditAnywhere, BlueprintReadOnly, Category = "Settings", meta = (TitleProperty = "DomainBodyNames"))
+	UPROPERTY(EditAnywhere, BlueprintReadWrite, Category = "Settings", meta = (TitleProperty = "DomainBodyNames"))
 	TArray<FName> DomainBodyNames;
 	
-	UPROPERTY(EditAnywhere, BlueprintReadOnly, Category = "Settings", meta = (TitleProperty = "ContactBodyNames"))
+	UPROPERTY(EditAnywhere, BlueprintReadWrite, Category = "Settings", meta = (TitleProperty = "ContactBodyNames"))
 	TArray<FName> ContactBodyNames;
 
 	// Relative Body Notify subclass to create
-	UPROPERTY(EditAnywhere, Category = "Settings", meta = (TitleProperty = "PhysicalIKRetargeterNotifyClass"))
+	UPROPERTY(EditAnywhere, BlueprintReadWrite, Category = "Settings", meta = (TitleProperty = "PhysicalIKRetargeterNotifyClass"))
 	TSubclassOf<URelativeBodyAnimNotifyBase> NotifyClass = nullptr;
 
 private:
