PreviewPropOp: Allow editor-only op that can preview props in retarget stack.

diff --git a/Engine/Plugins/Experimental/Animation/RelativeIKOp/RelativeIKOp.uplugin b/Engine/Plugins/Experimental/Animation/RelativeIKOp/RelativeIKOp.uplugin
--- a/Engine/Plugins/Experimental/Animation/RelativeIKOp/RelativeIKOp.uplugin
+++ b/Engine/Plugins/Experimental/Animation/RelativeIKOp/RelativeIKOp.uplugin
@@ -34,7 +34,12 @@
 			"Name": "RelativeBodyAnimUtils",
 			"Type": "Editor",
 			"LoadingPhase": "Default"
-		}
+		},
+		{
+			"Name": "PreviewPropOp",
+			"Type": "Runtime",
+			"LoadingPhase": "Default"
+		},
 	],
 	"Plugins": [
 		{
diff --git a/Engine/Plugins/Experimental/Animation/RelativeIKOp/Source/PreviewPropOp/PreviewPropOp.Build.cs b/Engine/Plugins/Experimental/Animation/RelativeIKOp/Source/PreviewPropOp/PreviewPropOp.Build.cs
new file mode 100644
--- /dev/null
+++ b/Engine/Plugins/Experimental/Animation/RelativeIKOp/Source/PreviewPropOp/PreviewPropOp.Build.cs
@@ -0,0 +1,61 @@
+// Copyright Epic Games, Inc. All Rights Reserved.
+
+using UnrealBuildTool;
+using System.IO; // for Path
+
+namespace UnrealBuildTool.Rules
+{
+	public class PreviewPropOp : ModuleRules
+	{
+		public PreviewPropOp(ReadOnlyTargetRules Target) : base(Target)
+		{
+			PCHUsage = ModuleRules.PCHUsageMode.UseExplicitOrSharedPCHs;
+			CppCompileWarningSettings.UnsafeTypeCastWarningLevel = WarningLevel.Error;
+
+			PublicIncludePaths.AddRange(
+				new string[] {
+					// ... add public include paths required here ...
+				}
+				);
+
+
+			PrivateIncludePaths.AddRange(
+				new string[] {
+					// ... add other private include paths required here ...
+				}
+				);
+
+
+			PublicDependencyModuleNames.AddRange(
+				new string[]
+				{
+				}
+				);
+
+
+			PrivateDependencyModuleNames.AddRange(
+				new string[]
+				{
+				"Core",
+				"CoreUObject",
+				"Engine",
+				"AnimationCore",
+				"MeshDescription",
+				"StaticMeshDescription",
+				"SkeletalMeshDescription",
+				"Renderer",
+				"IKRig"
+				}
+				);
+
+
+			DynamicallyLoadedModuleNames.AddRange(
+				new string[]
+				{
+					// ... add any modules that your module loads dynamically here ...
+				}
+				);
+		}
+	}
+
+}
\ No newline at end of file
diff --git a/Engine/Plugins/Experimental/Animation/RelativeIKOp/Source/PreviewPropOp/Private/PreviewPropOp.cpp b/Engine/Plugins/Experimental/Animation/RelativeIKOp/Source/PreviewPropOp/Private/PreviewPropOp.cpp
new file mode 100644
--- /dev/null
+++ b/Engine/Plugins/Experimental/Animation/RelativeIKOp/Source/PreviewPropOp/Private/PreviewPropOp.cpp
@@ -0,0 +1,339 @@
+// Copyright Epic Games, Inc. All Rights Reserved.
+
+#include "PreviewPropOp.h"
+
+#include "SceneView.h"
+#include "Retargeter/IKRetargetProcessor.h"
+
+#if WITH_EDITOR
+#include "PrimitiveDrawingUtils.h"
+#include "Engine/EngineTypes.h"
+#endif
+
+#include "StaticMeshAttributes.h"
+#include "AI/NavigationModifier.h"
+
+#include UE_INLINE_GENERATED_CPP_BY_NAME(PreviewPropOp)
+
+#define LOCTEXT_NAMESPACE "PreviewPropOp"
+
+const UClass* FIKRetargetPreviewPropOpSettings::GetControllerType() const
+{
+	return UIKRetargetPreviewPropController::StaticClass();
+}
+
+void FIKRetargetPreviewPropOpSettings::CopySettingsAtRuntime(const FIKRetargetOpSettingsBase* InSettingsToCopyFrom)
+{
+	static TArray<FName> PropertiesToIgnore = {"AttachMapping", "PreviewProps"};
+	FIKRetargetOpBase::CopyStructProperties(
+		FIKRetargetPreviewPropOpSettings::StaticStruct(),
+		InSettingsToCopyFrom,
+		this,
+		PropertiesToIgnore);
+}
+
+#if WITH_EDITOR
+USkeleton* FIKRetargetPreviewPropOpSettings::GetSkeleton(const FName InPropertyName)
+{
+	if (InPropertyName == GET_MEMBER_NAME_CHECKED(FPreviewPropsData, SourceAttachBone))
+	{
+		return const_cast<USkeleton*>(SourceSkeletonAsset);
+	}
+
+	return nullptr;
+}
+#endif //WITH_EDITOR
+
+bool FIKRetargetPreviewPropOp::Initialize(
+	const FIKRetargetProcessor& InProcessor,
+	const FRetargetSkeleton& InSourceSkeleton,
+	const FTargetSkeleton& InTargetSkeleton,
+	const FIKRetargetOpBase* InParentOp,
+	FIKRigLogger& InLog)
+{
+	bIsInitialized = false;
+#if WITH_EDITOR
+	PropMeshData.Reset();
+	
+	if (Settings.PreviewProps.IsEmpty())
+	{
+		return false;
+	}
+	
+	TObjectPtr<UMaterialInterface> PropDefMaterialInterface = Settings.DebugMaterial;
+	if (!PropDefMaterialInterface)
+	{
+		PropDefMaterialInterface = LoadObject<UMaterialInterface>(nullptr, TEXT("/Engine/EditorMaterials/Cloth/CameraLitDoubleSided.CameraLitDoubleSided"), nullptr, LOAD_None, nullptr);
+	}
+
+	if (!PropDefMaterialInterface)
+	{
+		return false;
+	}
+
+	PropDefaultMaterial = TStrongObjectPtr(UMaterialInstanceDynamic::Create(PropDefMaterialInterface, GetTransientPackage()));
+	if (!PropDefaultMaterial)
+	{
+		return false;
+	}
+	
+	for (const FPreviewPropsData& PreviewData : Settings.PreviewProps)
+	{
+		if (!PreviewData.PropMesh || PreviewData.SourceAttachBone.BoneName == NAME_None)
+		{
+			continue;
+		}
+		
+		int32 SourceAttachIdx = InSourceSkeleton.FindBoneIndexByName(PreviewData.SourceAttachBone.BoneName);
+		if (SourceAttachIdx == INDEX_NONE)
+		{
+			continue;
+		}
+
+		FName TargetAttachBone = ApplyBoneMap(PreviewData.SourceAttachBone.BoneName);
+		int32 TargetAttachIdx = InTargetSkeleton.FindBoneIndexByName(TargetAttachBone);
+
+		FPropMeshData NewMesh;
+		NewMesh.SourceAttachBoneIndex = SourceAttachIdx;
+		NewMesh.TargetAttachBoneIndex = TargetAttachIdx;
+		NewMesh.AttachTransform = PreviewData.AttachTransform;
+		
+		if (GetMeshVertTris(NewMesh, PreviewData))
+		{
+			FScopeLock ScopeLock(&DebugDataMutex);
+			PropMeshData.Add(NewMesh);
+		}
+	}
+
+	bIsInitialized = !PropMeshData.IsEmpty();
+#endif //WITH_EDITOR
+	return bIsInitialized;
+}
+
+void FIKRetargetPreviewPropOp::Run(
+	FIKRetargetProcessor& InProcessor,
+	const double InDeltaTime,
+	const TArray<FTransform>& InSourceGlobalPose,
+	TArray<FTransform>& OutTargetGlobalPose)
+{
+#if WITH_EDITOR
+	if (!bIsInitialized)
+	{
+		return;
+	}
+
+	double SourceScale = InProcessor.GetSourceScaleFactor();
+	FTransform SourceScaleTfm{FQuat::Identity, FVector::ZeroVector, FVector(SourceScale)};
+	
+	FScopeLock ScopeLock(&DebugDataMutex);
+	
+	DebugPropInfo.PropTfmInfo.Reset(PropMeshData.Num());
+
+	int32 PropIdx = 0;
+	for (const FPropMeshData& PropMesh : PropMeshData)
+	{
+		FPropTfmInfo& PropInfo = DebugPropInfo.PropTfmInfo.AddDefaulted_GetRef();
+		
+		FTransform SourceBoneTfm = InSourceGlobalPose[PropMesh.SourceAttachBoneIndex];
+		PropInfo.SourceTfm = PropMesh.AttachTransform * SourceScaleTfm * SourceBoneTfm;
+
+		PropInfo.bValidTargetTfm = (PropMesh.TargetAttachBoneIndex != INDEX_NONE);
+		if (PropInfo.bValidTargetTfm)
+		{
+			FTransform TargetBoneTfm = OutTargetGlobalPose[PropMesh.TargetAttachBoneIndex];
+			PropInfo.TargetTfm = PropMesh.AttachTransform * TargetBoneTfm;
+		}
+
+		++PropIdx;
+	}
+	
+#endif // WITH_EDITOR
+}
+
+FIKRetargetOpSettingsBase* FIKRetargetPreviewPropOp::GetSettings()
+{
+	return &Settings;
+}
+
+const UScriptStruct* FIKRetargetPreviewPropOp::GetSettingsType() const
+{
+	return FIKRetargetPreviewPropOpSettings::StaticStruct();
+}
+
+const UScriptStruct* FIKRetargetPreviewPropOp::GetType() const
+{
+	return FIKRetargetPreviewPropOp::StaticStruct();
+}
+
+FName FIKRetargetPreviewPropOp::ApplyBoneMap(FName SourceBone)
+{
+	if (Settings.AttachMapping.IsEmpty() || !Settings.AttachMapping.Contains(SourceBone))
+	{
+		return SourceBone;
+	}
+
+	return Settings.AttachMapping[SourceBone];
+}
+
+#if WITH_EDITOR
+// -----------------
+// All functionality is editor-only since at runtime 
+
+bool FIKRetargetPreviewPropOp::GetMeshVertTris(FPropMeshData& OutMeshData, const FPreviewPropsData& PreviewProp)
+{
+	const FMeshDescription* Desc = PreviewProp.PropMesh->GetMeshDescription(0);
+	if (!Desc)
+	{
+		return false;
+	}
+
+	int32 NumVerts = Desc->Vertices().Num();
+	int32 NumTris = Desc->Triangles().Num();
+	
+	OutMeshData.Vertices.Reset(NumVerts);
+	OutMeshData.Indices.SetNumUninitialized(3*NumTris);
+
+	// TODO: Almost certainly incorrectly "welding" all vert instances into whatever the first seen vert is!
+	TVertexInstanceAttributesConstRef<FVector2f> TexCoords = Desc->VertexInstanceAttributes().GetAttributesRef<FVector2f>(MeshAttribute::VertexInstance::TextureCoordinate);
+	TVertexInstanceAttributesConstRef<FVector3f> Tangents = Desc->VertexInstanceAttributes().GetAttributesRef<FVector3f>(MeshAttribute::VertexInstance::Tangent);
+	TVertexInstanceAttributesConstRef<FVector3f> Normals = Desc->VertexInstanceAttributes().GetAttributesRef<FVector3f>(MeshAttribute::VertexInstance::Normal);
+	TVertexInstanceAttributesConstRef<float> BinormalSigns = Desc->VertexInstanceAttributes().GetAttributesRef<float>(MeshAttribute::VertexInstance::BinormalSign);
+
+	int32 TriIdx = 0;
+	TMap<FVertexID, int32> TriVertMap;
+	for (const FTriangleID TriID : Desc->Triangles().GetElementIDs())
+	{
+		TConstArrayView<FVertexInstanceID> TriInst = Desc->GetTriangleVertexInstances(TriID);
+		TConstArrayView<FVertexID> TriVerts = Desc->GetTriangleVertices(TriID);
+
+		for (int32 Idx=0; Idx < 3; ++Idx)
+		{
+			const FVertexID VertID = TriVerts[Idx];
+			if (!TriVertMap.Contains(VertID))
+			{
+				int32 NewVertIdx = OutMeshData.Vertices.Num();
+
+				FDynamicMeshVertex NewVert{Desc->GetVertexPosition(VertID)};
+				if (TexCoords.IsValid())
+				{
+					// Max channels supported is 8
+					int32 TexChannels = FMath::Min(TexCoords.GetNumChannels(), 8);
+					for (int32 ChanIdx = 0; ChanIdx < TexChannels; ++ChanIdx)
+					{
+						NewVert.TextureCoordinate[ChanIdx] = TexCoords.Get(TriInst[Idx],ChanIdx);
+					}
+				}
+				if (Tangents.IsValid() && Normals.IsValid() && BinormalSigns.IsValid())
+				{
+					FVector3f Vtan = Tangents.Get(TriInst[Idx],0);
+					FVector3f Vnormal = Normals.Get(TriInst[Idx],0);
+					float Vsgn = BinormalSigns.Get(TriInst[Idx],0);
+					NewVert.TangentX = Vtan;
+					NewVert.TangentZ = Vnormal;
+					NewVert.TangentZ.Vector.W = static_cast<int8>(Vsgn);
+				}
+
+				OutMeshData.Vertices.Add(NewVert);
+				TriVertMap.Add(VertID, NewVertIdx);
+			}
+			OutMeshData.Indices[TriIdx*3 + Idx] = TriVertMap[VertID];
+		}
+		
+		++TriIdx;
+	}
+
+	return true;
+}
+
+
+
+FCriticalSection FIKRetargetPreviewPropOp::DebugDataMutex;
+
+void FIKRetargetPreviewPropOp::DebugDraw(
+	FPrimitiveDrawInterface* InPDI,
+	const FTransform& InSourceTransform,
+	const FTransform& InComponentTransform,
+	const double InComponentScale,
+	const FIKRetargetDebugDrawState& InEditorState) const
+{
+	FScopeLock ScopeLock(&DebugDataMutex);
+	
+	if (!Settings.bDebugDraw)
+	{
+		return;
+	}
+
+	if (!bIsInitialized)
+	{
+		return;
+	}
+
+	// Breakpoints or slow updates can have a debug call run between init/run calls
+	if (DebugPropInfo.PropTfmInfo.Num() != PropMeshData.Num())
+	{
+		return;
+	}
+
+	// TODO: Seems ridiculous to keep a strong and weak ptr, but this seems to stop crash on exit (look into GC of material)
+	if (!PropDefaultMaterial.IsValid())
+	{
+		return;
+	}
+
+	FTransform SourceNoScaleTfm = FTransform(InSourceTransform.GetRotation(), InSourceTransform.GetTranslation());
+
+	int PropIdx = 0;
+	for (const FPropTfmInfo& PropInfo : DebugPropInfo.PropTfmInfo)
+	{
+		const FPropMeshData& PropMesh = PropMeshData[PropIdx];
+		if (Settings.bShowSourceProps)
+		{
+			FDynamicMeshBuilder MeshBuilder(InPDI->View->GetFeatureLevel());
+			{
+				MeshBuilder.AddVertices(PropMesh.Vertices);
+				MeshBuilder.AddTriangles(PropMesh.Indices);
+			}
+			
+			FTransform SourceWorldTfm = DebugPropInfo.PropTfmInfo[PropIdx].SourceTfm * SourceNoScaleTfm;
+			
+			const FMaterialRenderProxy* MaterialProxy = PropDefaultMaterial->GetRenderProxy();
+			if (MaterialProxy)
+			{
+				MeshBuilder.Draw(InPDI, SourceWorldTfm.ToMatrixWithScale(), MaterialProxy, SDPG_World, false);
+			}
+		}
+
+		if (PropInfo.bValidTargetTfm)
+		{
+			FDynamicMeshBuilder MeshBuilder(InPDI->View->GetFeatureLevel());
+			{
+				MeshBuilder.AddVertices(PropMesh.Vertices);
+				MeshBuilder.AddTriangles(PropMesh.Indices);
+			}
+
+			FTransform TargetWorldTfm = DebugPropInfo.PropTfmInfo[PropIdx].TargetTfm * InComponentTransform;
+			
+			const FMaterialRenderProxy* MaterialProxy = PropDefaultMaterial->GetRenderProxy();
+			if (MaterialProxy)
+			{
+				MeshBuilder.Draw(InPDI, TargetWorldTfm.ToMatrixWithScale(), MaterialProxy, SDPG_World, false);
+			}
+		}
+		
+		++PropIdx;
+	}
+}
+
+#endif //WITH_EDITOR
+
+FIKRetargetPreviewPropOpSettings UIKRetargetPreviewPropController::GetSettings()
+{
+	return *reinterpret_cast<FIKRetargetPreviewPropOpSettings*>(OpSettingsToControl);
+}
+
+void UIKRetargetPreviewPropController::SetSettings(FIKRetargetPreviewPropOpSettings InSettings)
+{
+	OpSettingsToControl->CopySettingsAtRuntime(&InSettings);
+}
+
+#undef LOCTEXT_NAMESPACE
diff --git a/Engine/Plugins/Experimental/Animation/RelativeIKOp/Source/PreviewPropOp/Private/PreviewPropOpModule.cpp b/Engine/Plugins/Experimental/Animation/RelativeIKOp/Source/PreviewPropOp/Private/PreviewPropOpModule.cpp
new file mode 100644
--- /dev/null
+++ b/Engine/Plugins/Experimental/Animation/RelativeIKOp/Source/PreviewPropOp/Private/PreviewPropOpModule.cpp
@@ -0,0 +1,17 @@
+// Copyright Epic Games, Inc. All Rights Reserved.
+
+#include "PreviewPropOpModule.h"
+
+#define LOCTEXT_NAMESPACE "PreviewPropOpModule"
+
+void FPreviewPropOpModule::StartupModule()
+{
+}
+
+void FPreviewPropOpModule::ShutdownModule()
+{
+}
+
+#undef LOCTEXT_NAMESPACE
+    
+IMPLEMENT_MODULE(FPreviewPropOpModule, PreviewPropOp)
\ No newline at end of file
diff --git a/Engine/Plugins/Experimental/Animation/RelativeIKOp/Source/PreviewPropOp/Public/PreviewPropOp.h b/Engine/Plugins/Experimental/Animation/RelativeIKOp/Source/PreviewPropOp/Public/PreviewPropOp.h
new file mode 100644
--- /dev/null
+++ b/Engine/Plugins/Experimental/Animation/RelativeIKOp/Source/PreviewPropOp/Public/PreviewPropOp.h
@@ -0,0 +1,170 @@
+// Copyright Epic Games, Inc. All Rights Reserved.
+#pragma once
+
+#include "DynamicMeshBuilder.h"
+#include "Animation/BoneReference.h"
+#include "Engine/StaticMesh.h"
+#include "Materials/MaterialInstanceDynamic.h"
+#include "Retargeter/IKRetargetOps.h"
+
+#include "PreviewPropOp.generated.h"
+
+#define UE_API PREVIEWPROPOP_API
+
+#define LOCTEXT_NAMESPACE "PreviewPropOp"
+
+struct FIKRigGoal;
+struct FKShapeElem;
+class UPhysicsAsset;
+
+struct FPropTfmInfo
+{
+	bool bValidTargetTfm = false;
+	FTransform SourceTfm;
+	FTransform TargetTfm;
+};
+
+// Used to hold all debug draw info for convenience
+struct FDebugPropDrawInfo
+{
+	TArray<FPropTfmInfo> PropTfmInfo;
+};
+
+struct FPropMeshData
+{
+	TArray<uint32> Indices;
+	TArray<FDynamicMeshVertex> Vertices;
+
+	int32 SourceAttachBoneIndex;
+	int32 TargetAttachBoneIndex;
+	FTransform AttachTransform = FTransform::Identity;
+};
+
+USTRUCT(BlueprintType)
+struct FPreviewPropsData
+{
+	GENERATED_BODY()
+
+// #if WITH_EDITORONLY_DATA
+	// Prop mesh to preview
+	UPROPERTY(EditAnywhere, BlueprintReadWrite, Category = "Prop", meta=(ReinitializeOnEdit))
+	TObjectPtr<UStaticMesh> PropMesh;
+
+	// Source skeletal mesh bone to attach prop to
+	UPROPERTY(EditAnywhere, Category = "Prop", meta=(ReinitializeOnEdit))
+	FBoneReference SourceAttachBone;
+
+	// Local transform of prop relative to attach bone
+	UPROPERTY(EditAnywhere, BlueprintReadWrite, Category = "Prop", meta=(ReinitializeOnEdit))
+	FTransform AttachTransform;
+// ##endif //WITH_EDITORONLY_DATA
+};
+
+USTRUCT(BlueprintType, meta = (DisplayName = "Preview Prop Settings"))
+struct FIKRetargetPreviewPropOpSettings : public FIKRetargetOpSettingsBase
+{
+	GENERATED_BODY()
+
+	UE_API virtual const UClass* GetControllerType() const override;
+
+	UE_API virtual void CopySettingsAtRuntime(const FIKRetargetOpSettingsBase* InSettingsToCopyFrom) override;
+
+#if WITH_EDITOR
+	UE_API virtual USkeleton* GetSkeleton(const FName InPropertyName);
+#endif //WITH_EDITOR
+
+	// Attach bone source -> target name mapping 
+	UPROPERTY(EditAnywhere, BlueprintReadWrite, Category = "Source Target Mapping", meta=(ReinitializeOnEdit))
+	TMap<FName,FName> AttachMapping;
+	
+	// Prop info to preview when debug drawing enabled
+	UPROPERTY(EditAnywhere, BlueprintReadWrite, Category = "Debug", meta=(ReinitializeOnEdit))
+	TArray<FPreviewPropsData> PreviewProps;
+
+	// Debug Material for displaying props (must work w/ dynamic mesh, e.g. physics materials)
+	UPROPERTY(EditAnywhere, BlueprintReadWrite, Category = "Debug", meta=(ReinitializeOnEdit))
+	TObjectPtr<UMaterialInterface> DebugMaterial;
+
+	// Prop info to preview when debug drawing enabled
+	UPROPERTY(EditAnywhere, BlueprintReadWrite, Category = "Debug")
+	bool bShowSourceProps = false;
+};
+
+USTRUCT(BlueprintType, meta = (DisplayName = "Prop Debug Preview"))
+struct FIKRetargetPreviewPropOp : public FIKRetargetOpBase
+{
+	GENERATED_BODY()
+	
+	UE_API virtual bool Initialize(
+		const FIKRetargetProcessor& InProcessor,
+		const FRetargetSkeleton& InSourceSkeleton,
+		const FTargetSkeleton& InTargetSkeleton,
+		const FIKRetargetOpBase* InParentOp,
+		FIKRigLogger& InLog) override;
+	
+	UE_API virtual void Run(
+		FIKRetargetProcessor& InProcessor,
+		const double InDeltaTime,
+		const TArray<FTransform>& InSourceGlobalPose,
+		TArray<FTransform>& OutTargetGlobalPose) override;
+
+	UE_API virtual FIKRetargetOpSettingsBase* GetSettings() override;
+	
+	UE_API virtual const UScriptStruct* GetSettingsType() const override;
+	
+	UE_API virtual const UScriptStruct* GetType() const override;
+
+	UPROPERTY()
+	FIKRetargetPreviewPropOpSettings Settings;
+
+private:
+	FName ApplyBoneMap(FName SourceBone);
+
+#if WITH_EDITOR
+	static bool GetMeshVertTris(FPropMeshData& OutMeshData, const FPreviewPropsData& PreviewProp);
+	
+public:
+	UE_API virtual void DebugDraw(
+		FPrimitiveDrawInterface* InPDI,
+		const FTransform& InSourceTransform,
+		const FTransform& InComponentTransform,
+		const double InComponentScale,
+		const FIKRetargetDebugDrawState& InEditorState) const override;
+
+	virtual bool HasDebugDrawing() override { return true; };
+
+private:
+	// Bone transforms for prop preview
+	FDebugPropDrawInfo DebugPropInfo;
+	// Prop mesh info for all meshes (created in Initialize)
+	TArray<FPropMeshData> PropMeshData;
+	
+	static UE_API FCriticalSection DebugDataMutex;
+
+	
+	TStrongObjectPtr<UMaterialInstanceDynamic> PropDefaultMaterial = nullptr;
+#endif
+};
+
+/* The blueprint/python API for editing BodyIntersectIK Op */
+UCLASS(MinimalAPI, BlueprintType)
+class UIKRetargetPreviewPropController : public UIKRetargetOpControllerBase
+{
+	GENERATED_BODY()
+	
+public:
+
+	/* Get the current op settings as a struct.
+	 * @return FIKRetargetBodyIntersectIKOp struct with the current settings used by the op. */
+	UFUNCTION(BlueprintCallable, Category = Settings)
+	UE_API FIKRetargetPreviewPropOpSettings GetSettings();
+
+	/* Set the solver settings. Input is a custom struct type for this solver.
+	 * @param InSettings a FIKRetargetBodyIntersectIKOp struct containing all the settings to apply to this op */
+	UFUNCTION(BlueprintCallable, Category = Settings)
+	UE_API void SetSettings(FIKRetargetPreviewPropOpSettings InSettings);
+};
+
+#undef LOCTEXT_NAMESPACE
+
+#undef UE_API
diff --git a/Engine/Plugins/Experimental/Animation/RelativeIKOp/Source/PreviewPropOp/Public/PreviewPropOpModule.h b/Engine/Plugins/Experimental/Animation/RelativeIKOp/Source/PreviewPropOp/Public/PreviewPropOpModule.h
new file mode 100644
--- /dev/null
+++ b/Engine/Plugins/Experimental/Animation/RelativeIKOp/Source/PreviewPropOp/Public/PreviewPropOpModule.h
@@ -0,0 +1,14 @@
+// Copyright Epic Games, Inc. All Rights Reserved.
+
+#pragma once
+
+#include "Modules/ModuleManager.h"
+
+class FPreviewPropOpModule : public IModuleInterface
+{
+public:
+
+	/** IModuleInterface implementation */
+	virtual void StartupModule() override;
+	virtual void ShutdownModule() override;
+};
