RelativeIK: Merge prop integration and bug fixes to Main

diff --git a/Engine/Plugins/Experimental/Animation/RelativeIKOp/RelativeIKOp.uplugin b/Engine/Plugins/Experimental/Animation/RelativeIKOp/RelativeIKOp.uplugin
--- a/Engine/Plugins/Experimental/Animation/RelativeIKOp/RelativeIKOp.uplugin
+++ b/Engine/Plugins/Experimental/Animation/RelativeIKOp/RelativeIKOp.uplugin
@@ -45,6 +45,10 @@
 		{
 			"Name": "IKRig",
 			"Enabled": true
+		},
+		{
+			"Name": "GeometryProcessing",
+			"Enabled": true
 		}
 	]
 }
\ No newline at end of file
diff --git a/Engine/Plugins/Experimental/Animation/RelativeIKOp/Source/BodyIntersectIKOp/Private/BodyIntersectIKOp.cpp b/Engine/Plugins/Experimental/Animation/RelativeIKOp/Source/BodyIntersectIKOp/Private/BodyIntersectIKOp.cpp
--- a/Engine/Plugins/Experimental/Animation/RelativeIKOp/Source/BodyIntersectIKOp/Private/BodyIntersectIKOp.cpp
+++ b/Engine/Plugins/Experimental/Animation/RelativeIKOp/Source/BodyIntersectIKOp/Private/BodyIntersectIKOp.cpp
@@ -14,6 +14,8 @@
 #include "Engine/EngineTypes.h"
 #endif
 
+#include UE_INLINE_GENERATED_CPP_BY_NAME(BodyIntersectIKOp)
+
 #define LOCTEXT_NAMESPACE "BodyIntersectIKOp"
 
 
@@ -507,4 +509,4 @@ void UIKRetargetBodyIntersectController::SetSettings(FIKRetargetBodyIntersectIKO
 	OpSettingsToControl->CopySettingsAtRuntime(&InSettings);
 }
 
-#undef LOCTEXT_NAMESPACE
\ No newline at end of file
+#undef LOCTEXT_NAMESPACE
diff --git a/Engine/Plugins/Experimental/Animation/RelativeIKOp/Source/BodyIntersectIKOp/Public/BodyIntersectIKOp.h b/Engine/Plugins/Experimental/Animation/RelativeIKOp/Source/BodyIntersectIKOp/Public/BodyIntersectIKOp.h
--- a/Engine/Plugins/Experimental/Animation/RelativeIKOp/Source/BodyIntersectIKOp/Public/BodyIntersectIKOp.h
+++ b/Engine/Plugins/Experimental/Animation/RelativeIKOp/Source/BodyIntersectIKOp/Public/BodyIntersectIKOp.h
@@ -115,7 +115,7 @@ public:
 		const double InComponentScale,
 		const FIKRetargetDebugDrawState& InEditorState) const override;
 
-	UE_API virtual bool HasDebugDrawing() override { return true; };
+	virtual bool HasDebugDrawing() override { return true; };
 
 private:
 	void DebugDrawPhysBody(FPrimitiveDrawInterface* InPDI, const FTransform& ParentTransform, double Scale, const FKShapeElem* ShapeElem, const FLinearColor& Color) const;
diff --git a/Engine/Plugins/Experimental/Animation/RelativeIKOp/Source/RelativeBodyAnimInfo/Private/RelativeBodyAnimNotifies.cpp b/Engine/Plugins/Experimental/Animation/RelativeIKOp/Source/RelativeBodyAnimInfo/Private/RelativeBodyAnimNotifies.cpp
--- a/Engine/Plugins/Experimental/Animation/RelativeIKOp/Source/RelativeBodyAnimInfo/Private/RelativeBodyAnimNotifies.cpp
+++ b/Engine/Plugins/Experimental/Animation/RelativeIKOp/Source/RelativeBodyAnimInfo/Private/RelativeBodyAnimNotifies.cpp
@@ -2,3 +2,5 @@
 
 #include "RelativeBodyAnimNotifies.h"
 
+#include UE_INLINE_GENERATED_CPP_BY_NAME(RelativeBodyAnimNotifies)
+
diff --git a/Engine/Plugins/Experimental/Animation/RelativeIKOp/Source/RelativeBodyAnimInfo/Public/RelativeBodyAnimNotifies.h b/Engine/Plugins/Experimental/Animation/RelativeIKOp/Source/RelativeBodyAnimInfo/Public/RelativeBodyAnimNotifies.h
--- a/Engine/Plugins/Experimental/Animation/RelativeIKOp/Source/RelativeBodyAnimInfo/Public/RelativeBodyAnimNotifies.h
+++ b/Engine/Plugins/Experimental/Animation/RelativeIKOp/Source/RelativeBodyAnimInfo/Public/RelativeBodyAnimNotifies.h
@@ -6,11 +6,13 @@
 
 #include "RelativeBodyAnimNotifies.generated.h"
 
+#define UE_API RELATIVEBODYANIMINFO_API
+
 /**
  * Base class to limit options for notify class types in RelativeBodyAnimModifier
  */
-UCLASS()
-class RELATIVEBODYANIMINFO_API URelativeBodyAnimNotifyBase : public UAnimNotify
+UCLASS(MinimalAPI)
+class URelativeBodyAnimNotifyBase : public UAnimNotify
 {
 	GENERATED_BODY()
 };
@@ -19,8 +21,8 @@ class RELATIVEBODYANIMINFO_API URelativeBodyAnimNotifyBase : public UAnimNotify
 /**
  * Baked relative body anim-notify storing all body relationships in single notify at anim start
  */
-UCLASS()
-class RELATIVEBODYANIMINFO_API URelativeBodyBakeAnimNotify : public URelativeBodyAnimNotifyBase
+UCLASS(MinimalAPI)
+class URelativeBodyBakeAnimNotify : public URelativeBodyAnimNotifyBase
 {
 	GENERATED_BODY()
 	
@@ -39,7 +41,7 @@ public:
 	 * @param BodyPairsIsParentDominatesIn Array indicating if parent body should be allowed to move (one entery per-body-pair) 
 	 */
 	UFUNCTION(BlueprintImplementableEvent, Category = "Animation")
-	void OnRelativeBodyDenseNotify(USkeletalMeshComponent* MeshComp, UAnimSequenceBase* Animation, int NumSamplesIn, const TArray<FName> &BodyPairsIn, const TArray<float> &BodyPairsSampleTimeIn, const TArray<FVector3f> &BodyPairsLocalReferenceIn, const TArray<bool> &BodyPairsIsParentDominatesIn) const;
+	UE_API void OnRelativeBodyDenseNotify(USkeletalMeshComponent* MeshComp, UAnimSequenceBase* Animation, int NumSamplesIn, const TArray<FName> &BodyPairsIn, const TArray<float> &BodyPairsSampleTimeIn, const TArray<FVector3f> &BodyPairsLocalReferenceIn, const TArray<bool> &BodyPairsIsParentDominatesIn) const;
 
 	virtual void Notify(USkeletalMeshComponent* MeshComp, UAnimSequenceBase* Animation, const FAnimNotifyEventReference& EventReference) override
 	{
@@ -66,8 +68,8 @@ public:
 /**
  * Per frame body relationships notify (generates a very large amount of data)
  */
-UCLASS(meta = (DisplayName = "RelativeBodyPerFrameAnimNotify"))
-class RELATIVEBODYANIMINFO_API URelativeBodyPerFrameAnimNotify : public URelativeBodyAnimNotifyBase
+UCLASS(MinimalAPI, meta = (DisplayName = "RelativeBodyPerFrameAnimNotify"))
+class URelativeBodyPerFrameAnimNotify : public URelativeBodyAnimNotifyBase
 {
 	GENERATED_BODY()
 
@@ -93,7 +95,7 @@ public:
 	 * @param Loc2 The local reference position of contact point in the contact body
 	 */
 	UFUNCTION(BlueprintImplementableEvent, Category = "Animation")
-	void OnRelativeBodyAnimNotify(USkeletalMeshComponent* MeshComp, UAnimSequenceBase* Animation, USkeletalMesh* SkeletalMesh, FName Body1, FName Body2, const FVector3f &Loc1, const FVector3f &Loc2, bool IsParentDominates) const;
+	UE_API void OnRelativeBodyAnimNotify(USkeletalMeshComponent* MeshComp, UAnimSequenceBase* Animation, USkeletalMesh* SkeletalMesh, FName Body1, FName Body2, const FVector3f &Loc1, const FVector3f &Loc2, bool IsParentDominates) const;
 
 	// /**
 	//  * Restore baked local reference points to the source body's world position
@@ -131,3 +133,53 @@ protected:
 	UPROPERTY(VisibleAnywhere, AdvancedDisplay, meta = (Category = "Detail"))
     bool bIsParentDominates;
 };
+
+
+/**
+ * Baked relative props anim-notify storing all props relationships in single notify at anim start
+ */
+UCLASS(MinimalAPI)
+class URelativePropsBakeAnimNotify : public URelativeBodyAnimNotifyBase
+{
+	GENERATED_BODY()
+
+public:
+
+	// TODO: Use physics asset (possibly from override) instead of mesh component?
+	/**
+	 * OnRelativeBodyDenseNotify interface for contact pair info
+	 *
+	 * @param MeshComp The source SkeletalMeshComponent
+	 * @param Animation The current animation
+	 * @param NumSamplesIn The number of samples in baked data
+	 * @param BodyPairsIn Array of body pairs
+	 * @param BodyPairsSampleTimeIn Array of sample times
+	 * @param BodyPairsLocalReferenceIn Array of per-body (pair), per-sample, local reference positions of closest body verts
+	 * @param BodyPairsIsParentDominatesIn Array indicating if parent body should be allowed to move (one entery per-body-pair)
+	 */
+	UFUNCTION(BlueprintImplementableEvent, Category = "Animation")
+	UE_API void OnRelativePropsDenseNotify(USkeletalMeshComponent* MeshComp, UAnimSequenceBase* Animation, int NumSamplesIn, const TArray<float>& PropsPairsSampleTimeIn, const TArray<FTransform>& PropsOffsetTransformsIn, const TArray<FName>& PropsPairsIn, const TArray<FVector3f>& PropsPairsLocalReferenceIn) const;
+
+	virtual void Notify(USkeletalMeshComponent* MeshComp, UAnimSequenceBase* Animation, const FAnimNotifyEventReference& EventReference) override
+	{
+		// UE_LOG(LogTemp, Log, TEXT("URelativeIKDenseAnimNotify Triggered: %d constraints"), NumSamples);
+		OnRelativePropsDenseNotify(MeshComp, Animation, NumSamples, PropsPairsSampleTime, PropsOffsetTransforms, PropsPairs, PropsPairsLocalReference);
+	}
+
+	UPROPERTY(VisibleAnywhere, AdvancedDisplay, meta = (Category = "Detail"))
+	int NumSamples = 0;
+
+	UPROPERTY(VisibleAnywhere, AdvancedDisplay, meta = (Category = "Detail"))
+	TArray<float> PropsPairsSampleTime;
+
+	UPROPERTY(VisibleAnywhere, AdvancedDisplay, meta = (Category = "Detail"))
+	TArray<FTransform> PropsOffsetTransforms;
+
+	UPROPERTY(VisibleAnywhere, AdvancedDisplay, meta = (Category = "Detail"))
+	TArray<FName> PropsPairs;
+
+	UPROPERTY(VisibleAnywhere, AdvancedDisplay, meta = (Category = "Detail"))
+	TArray<FVector3f> PropsPairsLocalReference;
+};
+
+#undef UE_API
diff --git a/Engine/Plugins/Experimental/Animation/RelativeIKOp/Source/RelativeBodyAnimUtils/Private/RelativePropsAnimModifier.cpp b/Engine/Plugins/Experimental/Animation/RelativeIKOp/Source/RelativeBodyAnimUtils/Private/RelativePropsAnimModifier.cpp
new file mode 100644
--- /dev/null
+++ b/Engine/Plugins/Experimental/Animation/RelativeIKOp/Source/RelativeBodyAnimUtils/Private/RelativePropsAnimModifier.cpp
@@ -0,0 +1,330 @@
+// Copyright Epic Games, Inc. All Rights Reserved.
+
+#include "RelativePropsAnimModifier.h"
+
+#include "AnimPose.h"
+#include "MeshDescription.h"
+#include "RelativeBodyAnimModifier.h"
+#include "RelativeBodyAnimNotifies.h"
+#include "RelativeBodyUtils.h"
+#include "SkeletalMeshAttributes.h"
+#include "Animation/AnimSequence.h"
+#include "Components/SkeletalMeshComponent.h"
+#include "Engine/SkeletalMesh.h"
+#include "Misc/ScopedSlowTask.h"
+#include "PhysicsEngine/PhysicsAsset.h"
+#include "PhysicsEngine/SkeletalBodySetup.h"
+#include "Rendering/SkeletalMeshRenderData.h"
+#include "Spatial/PointHashGrid3.h"
+
+#include "StaticMeshAttributes.h"
+#include "MeshQueries.h"
+#include "MeshDescriptionToDynamicMesh.h"
+#include "MeshAdapter.h"
+#include "Spatial/MeshAABBTree3.h"
+#include "Algo/MinElement.h"
+
+#define LOCTEXT_NAMESPACE "RelativePropsAnimModifier"
+
+int32 FindLowestIndex(const TArray<FVector3f>& Points)
+{
+	if (Points.Num() == 0)
+	{
+		return INDEX_NONE;
+	}
+
+	auto It = Algo::MinElement(
+		Points,
+		[](const FVector3f& A, const FVector3f& B)
+		{
+			return A.Z < B.Z;
+		}
+	);
+
+	return It ? UE_PTRDIFF_TO_INT32(It - Points.GetData()) : INDEX_NONE;
+}
+
+int32 FindHighestIndex(const TArray<FVector3f>& Points)
+{
+	if (Points.Num() == 0)
+	{
+		return INDEX_NONE;
+	}
+
+	auto It = Algo::MinElement(
+		Points,
+		[](const FVector3f& A, const FVector3f& B)
+		{
+			return A.Z > B.Z;
+		}
+	);
+
+	return It ? UE_PTRDIFF_TO_INT32(It - Points.GetData()) : INDEX_NONE;
+}
+
+struct FPropsMeshAndTree
+{
+	UE::Geometry::FDynamicMesh3 Mesh;
+	UE::Geometry::TMeshAABBTree3<UE::Geometry::FDynamicMesh3> AABBTree;
+
+	FPropsMeshAndTree()
+		: AABBTree(&Mesh, false) {}  // defer build
+
+	void Init(const FMeshDescription* MeshDesc)
+	{
+		FMeshDescriptionToDynamicMesh Converter;
+		Converter.Convert(MeshDesc, Mesh);
+		AABBTree.Build();
+	}
+};
+
+void URelativePropsAnimModifier::OnApply_Implementation(UAnimSequence* InAnimation)
+{
+	if (NotifyClass)
+	{
+		URelativeBodyAnimModifier::OnApply_Implementation(InAnimation);
+	}
+
+	// TODO: Make log category for relative body utils?
+	if (InAnimation == nullptr)
+	{
+		UE_LOG(LogAnimation, Error, TEXT("Cannot create RelativeBodyAnimNotify. Reason: Invalid Animation"));
+		return;
+	}
+
+	if (!SkeletalMeshAsset)
+	{
+		UE_LOG(LogAnimation, Error, TEXT("Cannot create RelativeBodyAnimNotify. Reason: Invalid SkeletalMeshComponent"));
+		return;
+	}
+
+	if (!SkeletalMeshAsset->CloneMeshDescription(LODIndex, MeshDescription))
+	{
+		UE_LOG(LogAnimation, Error, TEXT("Cannot create RelativeBodyAnimNotify. Reason: Could not clone mesh description"));
+		return;
+	}
+
+	if (!PropsNotifyClass || PropsNotifyClass == URelativeBodyAnimNotifyBase::StaticClass())
+	{
+		UE_LOG(LogAnimation, Error, TEXT("Cannot create RelativePropsAnimNotify. Reason: Must specify valid relative body notify subclass to create"));
+		return;
+	}
+
+	if (!PhysicsAssetOverride)
+	{
+		UE_LOG(LogAnimation, Warning, TEXT("RelativeBodyAnimModifier: No physics asset override specified, using skeletal mesh physics asset!"))
+	}
+
+	//Cache
+	UPhysicsAsset* PhysicsAsset = GetPhysicsAsset();
+	FReferenceSkeleton& RefSkeleton = SkeletalMeshAsset->GetRefSkeleton();
+
+	// TODO: Move cached data out of modifier (or check if it's outdated and only rebuild in that case)
+	CacheBodyDataForSourceMesh(CachedBodySourceData);
+
+	const TArray<int32>& BodyIndicesParentBodyIndices = CachedBodySourceData.BodyIndicesParentBodyIndices;
+	const TArray<bool>& BodyIndicesToIgnore = CachedBodySourceData.BodyIndicesToIgnore;
+	const TArray<bool>& IsDomainBodyIndices = CachedBodySourceData.IsDomainBody;
+	const TArray<int32>& BodyIndicesToBoneIndices = CachedBodySourceData.BodyIndicesToSourceBoneIndices;
+	const TArray<int32>& BoneIndicesToBodyIndices = CachedBodySourceData.SourceBoneIndicesToBodyIndices;
+	const TArray<TArray<int32>>& VertexIndicesInfluencedByBodyIndices = CachedBodySourceData.SourceVertexIndicesInfluencedByBodyIndices;
+
+	const int32 NumBodies = BodyIndicesToBoneIndices.Num();
+	const float HashGridCellSize = std::min(1.5f * ContactThreshold, 50.f); // the length of the cell size in the point hash gri
+
+	// Process animation asset
+	{
+		const FAnimPoseEvaluationOptions AnimPoseEvalOptions{
+			EAnimDataEvalType::Raw,
+			true,
+			false,
+			false,
+			nullptr,
+			true,
+			false
+		};
+
+		const float SequenceLength = InAnimation->GetPlayLength();
+		const float SampleStep = 1.0f / FMath::Max(1.0f, static_cast<float>(SampleRate));
+		const int SampleNum = FMath::TruncToInt(SequenceLength / SampleStep);
+
+		URelativePropsBakeAnimNotify* BakedAnimNotifyInfo = nullptr;
+
+		bool bUseDenseNotifyInfo = PropsNotifyClass->IsChildOf(URelativePropsBakeAnimNotify::StaticClass());
+		TArray<FPropsMeshAndTree> PropsTreeData;
+		PropsTreeData.SetNum(PropsStaticMeshs.Num());
+		if (bUseDenseNotifyInfo)
+		{
+			FName TrackName("RelProps-<" + SkeletalMeshAsset->GetName() + ">");
+			UE_LOG(LogAnimation, Verbose, TEXT("TrackName %d: '%s'."), GeneratedNotifyTracks.Num() - 1, *TrackName.ToString());
+			const bool bDoesTrackNameAlreadyExist = UAnimationBlueprintLibrary::IsValidAnimNotifyTrackName(InAnimation, TrackName);
+			if (!bDoesTrackNameAlreadyExist)
+			{
+				UAnimationBlueprintLibrary::AddAnimationNotifyTrack(InAnimation, TrackName, FLinearColor::MakeRandomColor());
+				GeneratedNotifyTracks.Add(TrackName);
+			}
+
+			BakedAnimNotifyInfo = Cast<URelativePropsBakeAnimNotify>(UAnimationBlueprintLibrary::AddAnimationNotifyEvent(InAnimation, TrackName, 0., PropsNotifyClass));
+			if (!BakedAnimNotifyInfo)
+			{
+				return;
+			}
+
+			for (size_t PropsIdx = 0; PropsIdx < PropsStaticMeshs.Num(); PropsIdx++)
+			{
+				FMeshDescription* PropsMeshDescription = PropsStaticMeshs[PropsIdx]->GetMeshDescription(0); // LOD 0
+				if (PropsMeshDescription)
+				{
+					PropsTreeData[PropsIdx].Init(PropsMeshDescription);
+				}
+
+				for (int32 BodyIndex1 = 0; BodyIndex1 < NumBodies; ++BodyIndex1)
+				{
+					if (BodyIndicesToIgnore[BodyIndex1])
+					{
+						continue;
+					}
+					FName BoneName = RefSkeleton.GetBoneName(BodyIndicesToBoneIndices[BodyIndex1]);
+					BakedAnimNotifyInfo->PropsPairs.Add(PropsAttachBones[PropsIdx]);
+					BakedAnimNotifyInfo->PropsPairs.Add(BoneName);
+				}
+				if (bPropsFloorInfoBaking)
+				{
+					BakedAnimNotifyInfo->PropsPairs.Add(PropsAttachBones[PropsIdx]);
+					BakedAnimNotifyInfo->PropsPairs.Add(FName("Floor"));
+				}
+			}
+			BakedAnimNotifyInfo->PropsPairsLocalReference.SetNumZeroed(BakedAnimNotifyInfo->PropsPairs.Num() * SampleNum);
+			BakedAnimNotifyInfo->PropsPairsSampleTime.SetNumZeroed(SampleNum);
+			BakedAnimNotifyInfo->NumSamples = SampleNum;
+			BakedAnimNotifyInfo->PropsOffsetTransforms = PropsOffsetTransforms;
+		}
+		
+		FScopedSlowTask BakePairVertsTask((float)SampleNum, LOCTEXT("BakeVertsTaskText", "Baking Props Pair Verts..."));
+		BakePairVertsTask.MakeDialog();
+
+		// Get ground levels and max speed values.
+		TArray<FVector3f> VLocations;
+		TArray<TArray<FVector3f>> PropsVLocations;
+		for (int SampleIndex = 0; SampleIndex < SampleNum; ++SampleIndex)
+		{
+			BakePairVertsTask.EnterProgressFrame(1.0f);
+
+			const float SampleTime = FMath::Clamp(static_cast<float>(SampleIndex) * SampleStep, 0.0f, SequenceLength);
+			const float FutureSampleTime = FMath::Clamp((static_cast<float>(SampleIndex) + 1.0f) * SampleStep, 0.0f, SequenceLength);
+
+			FAnimPose AnimPose;
+			FAnimPose FutureAnimPose;
+
+			UAnimPoseExtensions::GetAnimPoseAtTime(InAnimation, SampleTime, AnimPoseEvalOptions, AnimPose);
+			UAnimPoseExtensions::GetAnimPoseAtTime(InAnimation, FutureSampleTime, AnimPoseEvalOptions, FutureAnimPose);
+
+			TArray<FMatrix44f> CacheToLocals;
+			GetRefToAnimPoseMatrices(CacheToLocals, AnimPose);
+			GetSkinnedVertices(VLocations, CacheToLocals);
+			if (BakedAnimNotifyInfo)
+			{
+				GetPropsVertices(PropsVLocations, AnimPose);
+			}
+
+			size_t Count = 0;
+			TArray<FName> ContactedBodyPairs;
+			TArray<FVector3f> ContactPoints;
+			ContactPoints.Empty();
+			if (!BakedAnimNotifyInfo)
+			{
+				continue;
+			}
+			else 
+			{
+				BakedAnimNotifyInfo->PropsPairsSampleTime[SampleIndex] = SampleTime;
+				for (int32 PropsIdx = 0; PropsIdx < PropsStaticMeshs.Num(); PropsIdx++)
+				{
+					FTransform AttachBoneTransform = UAnimPoseExtensions::GetBonePose(AnimPose, PropsAttachBones[PropsIdx], EAnimPoseSpaces::World);
+					FTransform PropsLocalToWorld = PropsOffsetTransforms[PropsIdx]*AttachBoneTransform;
+					const FStaticMeshLODResources& LODResources = PropsStaticMeshs[PropsIdx]->GetRenderData()->LODResources[0];
+					const FPositionVertexBuffer& PositionVertexBuffer = LODResources.VertexBuffers.PositionVertexBuffer;
+					for (int32 BodyIndex = 0; BodyIndex < NumBodies; ++BodyIndex)
+					{
+						if (BodyIndicesToIgnore[BodyIndex])
+						{
+							continue;
+						}
+
+						size_t Offset = SampleIndex * BakedAnimNotifyInfo->PropsPairs.Num() + Count * 2;
+						int32 NumVerts = VertexIndicesInfluencedByBodyIndices[BodyIndex].Num();
+
+						double MinDist = DBL_MAX;
+						FVector3d NearestPoint;
+						int32 NearestQueryVertexIndex = INDEX_NONE;
+						for (int32 QueryVertexIndex = 0; QueryVertexIndex < NumVerts; ++QueryVertexIndex) 
+						{
+							FVector3d VertexPosition(VLocations[VertexIndicesInfluencedByBodyIndices[BodyIndex][QueryVertexIndex]]);
+							VertexPosition = PropsLocalToWorld.InverseTransformPosition(VertexPosition);
+							FVector3d NearestPointTemp = PropsTreeData[PropsIdx].AABBTree.FindNearestPoint(VertexPosition);
+							double DistSqr = FVector3d::DistSquared(VertexPosition, NearestPointTemp);
+							// UE_LOG(LogAnimation, Warning, TEXT("QueryVertexIndex %d: Dist=%f (%f,%f,%f)."), QueryVertexIndex, DistSqr, NearestPointTemp.X, NearestPointTemp.Y, NearestPointTemp.Z);
+							if (DistSqr < MinDist)
+							{
+								NearestQueryVertexIndex = QueryVertexIndex;
+								NearestPoint = NearestPointTemp;
+								MinDist = DistSqr;
+							}
+						}
+						// UE_LOG(LogAnimation, Warning, TEXT("MinDIst %f: (%f,%f,%f)."), MinDist, NearestPoint.X, NearestPoint.Y, NearestPoint.Z);
+						if (NearestQueryVertexIndex != INDEX_NONE)
+						{
+							FName BoneName1 = PropsAttachBones[PropsIdx];
+							FName BoneName2 = RefSkeleton.GetBoneName(BodyIndicesToBoneIndices[BodyIndex]);
+							FVector3f ContactPoint1(PropsLocalToWorld.TransformPosition(NearestPoint));
+							FVector3f ContactPoint2 = VLocations[VertexIndicesInfluencedByBodyIndices[BodyIndex][NearestQueryVertexIndex]];
+
+							FVector3f Loc1(NearestPoint);
+							FVector3f Loc2 = FRelativeBodyAnimUtils::CalcVertLocationInUnitBody(ContactPoint2, BoneName2, RefSkeleton, CacheToLocals, CachedBodySourceData.SourceRetargetGlobalPose, PhysicsAsset);
+
+							BakedAnimNotifyInfo->PropsPairsLocalReference[Offset + 0] = Loc1;
+							BakedAnimNotifyInfo->PropsPairsLocalReference[Offset + 1] = Loc2;
+							//UE_LOG(LogAnimation, Warning, TEXT("NumVertsIDX %d: (%f,%f,%f)."), NearestQueryVertexIndex, Loc1.X, Loc1.Y, Loc1.Z);
+						}
+						Count++;
+					}
+					if (bPropsFloorInfoBaking)
+					{
+						int32 Offset = SampleIndex * BakedAnimNotifyInfo->PropsPairs.Num() + Count * 2;
+						int32 LowestVertexIndex = FindLowestIndex(PropsVLocations[PropsIdx]);
+						int32 HighestVertexIndex=FindHighestIndex(PropsVLocations[PropsIdx]);
+
+						BakedAnimNotifyInfo->PropsPairsLocalReference[Offset + 0] = PositionVertexBuffer.VertexPosition(LowestVertexIndex);
+						BakedAnimNotifyInfo->PropsPairsLocalReference[Offset + 1] = PositionVertexBuffer.VertexPosition(HighestVertexIndex);
+
+						Count++;
+					}
+				}
+			}
+		}
+	}
+}
+
+void URelativePropsAnimModifier::GetPropsVertices(TArray<TArray<FVector3f>>& PropsVLocations, const FAnimPose& AnimPose)
+{
+	PropsVLocations.Empty();
+	for (int32 PropsIdx = 0; PropsIdx < PropsStaticMeshs.Num(); PropsIdx++)
+	{
+		const FTransform AttachBoneTransform = UAnimPoseExtensions::GetBonePose(AnimPose, PropsAttachBones[PropsIdx], EAnimPoseSpaces::World);
+
+		const FStaticMeshLODResources& LODResources = PropsStaticMeshs[PropsIdx]->GetRenderData()->LODResources[0];
+		const FPositionVertexBuffer& PositionVertexBuffer = LODResources.VertexBuffers.PositionVertexBuffer;
+
+		TArray<FVector3f> LocalVertices;
+		int32 NumVerts = PositionVertexBuffer.GetNumVertices();
+		for (int32 i = 0; i < NumVerts; ++i)
+		{
+			FVector LocalPos = PropsOffsetTransforms[PropsIdx].TransformPosition(FVector(PositionVertexBuffer.VertexPosition(i)));
+			LocalVertices.Add(FVector3f(AttachBoneTransform.TransformPosition(LocalPos)));
+		}
+		PropsVLocations.Add(LocalVertices);
+	}
+
+}
+
+#undef LOCTEXT_NAMESPACE
\ No newline at end of file
diff --git a/Engine/Plugins/Experimental/Animation/RelativeIKOp/Source/RelativeBodyAnimUtils/Public/RelativeBodyAnimModifier.h b/Engine/Plugins/Experimental/Animation/RelativeIKOp/Source/RelativeBodyAnimUtils/Public/RelativeBodyAnimModifier.h
--- a/Engine/Plugins/Experimental/Animation/RelativeIKOp/Source/RelativeBodyAnimUtils/Public/RelativeBodyAnimModifier.h
+++ b/Engine/Plugins/Experimental/Animation/RelativeIKOp/Source/RelativeBodyAnimUtils/Public/RelativeBodyAnimModifier.h
@@ -124,7 +124,7 @@ public:
 	UPROPERTY(EditAnywhere, BlueprintReadWrite, Category = "Settings", meta = (TitleProperty = "PhysicalIKRetargeterNotifyClass"))
 	TSubclassOf<URelativeBodyAnimNotifyBase> NotifyClass = nullptr;
 
-private:
+protected:
 	FMeshDescription MeshDescription;
 	FRelativeBodySourceData CachedBodySourceData;
 
diff --git a/Engine/Plugins/Experimental/Animation/RelativeIKOp/Source/RelativeBodyAnimUtils/Public/RelativePropsAnimModifier.h b/Engine/Plugins/Experimental/Animation/RelativeIKOp/Source/RelativeBodyAnimUtils/Public/RelativePropsAnimModifier.h
new file mode 100644
--- /dev/null
+++ b/Engine/Plugins/Experimental/Animation/RelativeIKOp/Source/RelativeBodyAnimUtils/Public/RelativePropsAnimModifier.h
@@ -0,0 +1,50 @@
+// Copyright Epic Games, Inc. All Rights Reserved.
+
+#pragma once
+
+#include "RelativeBodyAnimModifier.h"
+#include "MeshDescription.h"
+#include "UObject/ObjectPtr.h"
+
+#include "RelativePropsAnimModifier.generated.h"
+
+class UPhysicsAsset;
+class URelativeBodyAnimNotifyBase;
+class USkeletalMesh;
+struct FAnimPose;
+struct FReferenceSkeleton;
+
+/**
+ * Animation modifier for baking relative body relationships into anim notify
+ */
+UCLASS(meta = (IsBlueprintBase = true))
+class RELATIVEBODYANIMUTILS_API URelativePropsAnimModifier : public URelativeBodyAnimModifier
+{
+	GENERATED_BODY()
+
+public:
+	/** Begin UAnimationModifier interface */
+	virtual void OnApply_Implementation(UAnimSequence* InAnimation) override;
+	/** End UAnimationModifier interface */
+
+public:
+	
+	UPROPERTY(EditAnywhere, BlueprintReadOnly, Category = "Settings", meta = (TitleProperty = "PropsStaticMeshs"))
+	TArray<TObjectPtr<UStaticMesh>> PropsStaticMeshs;
+
+	UPROPERTY(EditAnywhere, BlueprintReadOnly, Category = "Settings", meta = (TitleProperty = "PropsAttachBones"))
+	TArray<FName> PropsAttachBones;
+
+	/** The local offset transform that defines the attachment position on the props static mesh */
+	UPROPERTY(EditAnywhere, BlueprintReadOnly, Category = "Settings", meta = (TitleProperty = "PropsOffsetTransforms"))
+	TArray<FTransform> PropsOffsetTransforms;
+
+	/**  Relative Props Notify subclass to create */
+	UPROPERTY(EditAnywhere, Category = "Settings", meta = (TitleProperty = "PropsNotifyClass"))
+	TSubclassOf<URelativeBodyAnimNotifyBase> PropsNotifyClass = nullptr;
+
+	UPROPERTY(EditAnywhere, Category = "Settings", meta = (TitleProperty = "PropsFloorInfoBaking"))
+	bool bPropsFloorInfoBaking = false;
+protected:
+	void GetPropsVertices(TArray<TArray<FVector3f>>& PropsVLocations, const FAnimPose& AnimPose);
+};
\ No newline at end of file
diff --git a/Engine/Plugins/Experimental/Animation/RelativeIKOp/Source/RelativeBodyAnimUtils/RelativeBodyAnimUtils.Build.cs b/Engine/Plugins/Experimental/Animation/RelativeIKOp/Source/RelativeBodyAnimUtils/RelativeBodyAnimUtils.Build.cs
--- a/Engine/Plugins/Experimental/Animation/RelativeIKOp/Source/RelativeBodyAnimUtils/RelativeBodyAnimUtils.Build.cs
+++ b/Engine/Plugins/Experimental/Animation/RelativeIKOp/Source/RelativeBodyAnimUtils/RelativeBodyAnimUtils.Build.cs
@@ -40,6 +40,9 @@ namespace UnrealBuildTool.Rules
 				"CoreUObject",
 				"Engine",
 				"AnimationCore",
+				"GeometryCore",
+				"DynamicMesh",
+				"MeshConversion",
 				"MeshDescription",
 				"StaticMeshDescription",
 				"SkeletalMeshDescription",
diff --git a/Engine/Plugins/Experimental/Animation/RelativeIKOp/Source/RelativeIKOp/Private/RelativeIKOp.cpp b/Engine/Plugins/Experimental/Animation/RelativeIKOp/Source/RelativeIKOp/Private/RelativeIKOp.cpp
--- a/Engine/Plugins/Experimental/Animation/RelativeIKOp/Source/RelativeIKOp/Private/RelativeIKOp.cpp
+++ b/Engine/Plugins/Experimental/Animation/RelativeIKOp/Source/RelativeIKOp/Private/RelativeIKOp.cpp
@@ -20,6 +20,8 @@
 #include "Engine/EngineTypes.h"
 #endif
 
+#include UE_INLINE_GENERATED_CPP_BY_NAME(RelativeIKOp)
+
 #define LOCTEXT_NAMESPACE "RelativeIKOp"
 
 const UClass* FIKRetargetRelativeIKOpSettings::GetControllerType() const
@@ -1358,4 +1360,4 @@ void UIKRetargetRelativeIKController::SetSettings(FIKRetargetRelativeIKOpSetting
 	OpSettingsToControl->CopySettingsAtRuntime(&InSettings);
 }
 
-#undef LOCTEXT_NAMESPACE
\ No newline at end of file
+#undef LOCTEXT_NAMESPACE
diff --git a/Engine/Plugins/Experimental/Animation/RelativeIKOp/Source/RelativeIKOp/Public/RelativeIKOp.h b/Engine/Plugins/Experimental/Animation/RelativeIKOp/Source/RelativeIKOp/Public/RelativeIKOp.h
--- a/Engine/Plugins/Experimental/Animation/RelativeIKOp/Source/RelativeIKOp/Public/RelativeIKOp.h
+++ b/Engine/Plugins/Experimental/Animation/RelativeIKOp/Source/RelativeIKOp/Public/RelativeIKOp.h
@@ -275,7 +275,7 @@ public:
 		const double InComponentScale,
 		const FIKRetargetDebugDrawState& InEditorState) const override;
 
-	UE_API virtual bool HasDebugDrawing() override { return true; };
+	virtual bool HasDebugDrawing() override { return true; };
 
 private:
 	void DebugDrawBodyPairs(FPrimitiveDrawInterface* InPDI, const FTransform& BaseTransform, double Scale, const TArray<FDebugDrawBodyPair>& BodyPairs) const;
diff --git a/Engine/Plugins/Experimental/Animation/RelativeIKOp/Source/RelativeIKOp/RelativeIKOp.Build.cs b/Engine/Plugins/Experimental/Animation/RelativeIKOp/Source/RelativeIKOp/RelativeIKOp.Build.cs
--- a/Engine/Plugins/Experimental/Animation/RelativeIKOp/Source/RelativeIKOp/RelativeIKOp.Build.cs
+++ b/Engine/Plugins/Experimental/Animation/RelativeIKOp/Source/RelativeIKOp/RelativeIKOp.Build.cs
@@ -40,6 +40,9 @@ namespace UnrealBuildTool.Rules
 				"CoreUObject",
 				"Engine",
 				"AnimationCore",
+				"GeometryCore",
+				"DynamicMesh",
+				"MeshConversion",
 				"MeshDescription",
 				"StaticMeshDescription",
 				"SkeletalMeshDescription",
