RelativeIKOp: Don't ignore physics assets in settings copy.

diff --git a/Engine/Plugins/Experimental/Animation/RelativeIKOp/Source/RelativeIKOp/Private/RelativeIKOp.cpp b/Engine/Plugins/Experimental/Animation/RelativeIKOp/Source/RelativeIKOp/Private/RelativeIKOp.cpp
--- a/Engine/Plugins/Experimental/Animation/RelativeIKOp/Source/RelativeIKOp/Private/RelativeIKOp.cpp
+++ b/Engine/Plugins/Experimental/Animation/RelativeIKOp/Source/RelativeIKOp/Private/RelativeIKOp.cpp
@@ -31,7 +31,7 @@ const UClass* FIKRetargetRelativeIKOpSettings::GetControllerType() const
 
 void FIKRetargetRelativeIKOpSettings::CopySettingsAtRuntime(const FIKRetargetOpSettingsBase* InSettingsToCopyFrom)
 {
-	static TArray<FName> PropertiesToIgnore = {"SourcePhysicsAssetOverride", "TargetPhysicsAssetOverride", "BodyMapping"};
+	static TArray<FName> PropertiesToIgnore = {};
 	FIKRetargetOpBase::CopyStructProperties(
 		FIKRetargetRelativeIKOpSettings::StaticStruct(),
 		InSettingsToCopyFrom,
