BodyIntersectIKOp: Added  custom sphere/capsule checks to fix intersection delta limit when sphere centers are inside intersect body.

diff --git a/Engine/Plugins/Experimental/Animation/RelativeIKOp/Source/BodyIntersectIKOp/Private/BodyIntersectIKOp.cpp b/Engine/Plugins/Experimental/Animation/RelativeIKOp/Source/BodyIntersectIKOp/Private/BodyIntersectIKOp.cpp
--- a/Engine/Plugins/Experimental/Animation/RelativeIKOp/Source/BodyIntersectIKOp/Private/BodyIntersectIKOp.cpp
+++ b/Engine/Plugins/Experimental/Animation/RelativeIKOp/Source/BodyIntersectIKOp/Private/BodyIntersectIKOp.cpp
@@ -257,20 +257,19 @@ double FIKRetargetBodyIntersectIKOp::CalcIntersectionDelta(const FVector& Target
 	}
 
 	float Dist = 0.0f;
-	FVector ClosestPoint;
 	switch (ShapeElem->GetShapeType())
 	{
 	case EAggCollisionShape::Sphere:
 		{
 			FKSphereElem* SphereElem = static_cast<FKSphereElem*>(ShapeElem);
-			Dist = SphereElem->GetClosestPointAndNormal(TargetPoint, ShapeTfm, ClosestPoint, OutDeltaDir);
+			Dist = CalcPointShapeDist(OutDeltaDir, TargetPoint, ShapeTfm, SphereElem);
 			break;
 		}
 
 	case EAggCollisionShape::Sphyl:
 		{
 			FKSphylElem* CapsuleElem = static_cast<FKSphylElem*>(ShapeElem);
-			Dist = CapsuleElem->GetClosestPointAndNormal(TargetPoint, ShapeTfm, ClosestPoint, OutDeltaDir);
+			Dist = CalcPointShapeDist(OutDeltaDir, TargetPoint, ShapeTfm, CapsuleElem);
 			break;
 		}
 
@@ -290,6 +289,41 @@ double FIKRetargetBodyIntersectIKOp::CalcIntersectionDelta(const FVector& Target
 	return TargetRadius - Dist;
 }
 
+float FIKRetargetBodyIntersectIKOp::CalcPointShapeDist(FVector& OutDeltaDir, const FVector& TargetPoint, const FTransform& ShapeTfm, const FKSphereElem* SphereElem)
+{
+	const float SphereRadius = static_cast<float>(ShapeTfm.GetScale3D().GetAbsMin()) * SphereElem->Radius;
+	const FVector WorldCenter = ShapeTfm.TransformPosition(SphereElem->Center);
+
+	const FVector Dir = WorldCenter - TargetPoint;
+	const float DistToCenter = static_cast<float>(Dir.Size());
+	const float DistToEdge = DistToCenter - SphereRadius;
+
+	OutDeltaDir = (DistToCenter > UE_SMALL_NUMBER) ? -Dir.GetUnsafeNormal() : FVector::ZeroVector;
+	return DistToEdge;
+}
+
+float FIKRetargetBodyIntersectIKOp::CalcPointShapeDist(FVector& OutDeltaDir, const FVector& TargetPoint, const FTransform& ShapeTfm, const FKSphylElem* CapsuleElem)
+{
+	const FVector Scale3DAbs = ShapeTfm.GetScale3D().GetAbs();
+	const float CapsuleRadius = CapsuleElem->GetScaledRadius(Scale3DAbs);
+	const float CapsuleHalfLen = 0.5f * CapsuleElem->GetScaledCylinderLength(Scale3DAbs);
+
+	const FTransform LocalToWorldTfm = CapsuleElem->GetTransform() * ShapeTfm;
+	const FVector LocalTarget = LocalToWorldTfm.InverseTransformPositionNoScale(TargetPoint);
+
+	const double TargetZ = FMath::Clamp(LocalTarget.Z, -CapsuleHalfLen, CapsuleHalfLen);
+	const FVector WorldCenter = LocalToWorldTfm.TransformPositionNoScale(FVector(0.0,0.0, TargetZ));
+	
+	const FVector Dir = WorldCenter - TargetPoint;
+	const float DistToCenter = static_cast<float>(Dir.Size());
+	const float DistToEdge = DistToCenter - CapsuleRadius;
+
+	OutDeltaDir = (DistToCenter > UE_SMALL_NUMBER) ? -Dir.GetUnsafeNormal() : FVector::ZeroVector;
+	return DistToEdge;
+}
+
+
+
 double FIKRetargetBodyIntersectIKOp::CalcShapeSmallRadius(FKShapeElem* ShapeElem)
 {
 	// Return 1.0 as default so scaling can set "size" of sphere
diff --git a/Engine/Plugins/Experimental/Animation/RelativeIKOp/Source/BodyIntersectIKOp/Public/BodyIntersectIKOp.h b/Engine/Plugins/Experimental/Animation/RelativeIKOp/Source/BodyIntersectIKOp/Public/BodyIntersectIKOp.h
--- a/Engine/Plugins/Experimental/Animation/RelativeIKOp/Source/BodyIntersectIKOp/Public/BodyIntersectIKOp.h
+++ b/Engine/Plugins/Experimental/Animation/RelativeIKOp/Source/BodyIntersectIKOp/Public/BodyIntersectIKOp.h
@@ -12,6 +12,8 @@
 
 struct FIKRigGoal;
 struct FKShapeElem;
+struct FKSphereElem;
+struct FKSphylElem;
 class UPhysicsAsset;
 
 // Used to hold all debug draw info for convenience
@@ -99,7 +101,9 @@ private:
 	FVector GoalLocBlendCompSpace(const FIKRigGoal* Goal, const FTransform& BoneTfm) const;
 	void SetGoalPosFromCompSpace(FIKRigGoal* Goal, const FTransform& BoneTfm, const FVector& CompSpaceLoc) const;
 	
-	static double CalcIntersectionDelta(const FVector&  TargetPoint, double TargetRadius, const FTransform& ShapeTfm, FKShapeElem* ShapeElem, FVector& OutDeltaDir);
+	static double CalcIntersectionDelta(const FVector& TargetPoint, double TargetRadius, const FTransform& ShapeTfm, FKShapeElem* ShapeElem, FVector& OutDeltaDir);
+	static float CalcPointShapeDist(FVector& OutDeltaDir, const FVector& TargetPoint, const FTransform& ShapeTfm, const FKSphereElem* SphereElem);
+	static float CalcPointShapeDist(FVector& OutDeltaDir, const FVector& TargetPoint, const FTransform& ShapeTfm, const FKSphylElem* CapsuleElem);
 	
 	static double CalcShapeSmallRadius(FKShapeElem* ShapeElem);
 	static double CalcShapeAvgRadius(FKShapeElem* ShapeElem);
