RelativeIKOp: Integrate prop notifies and fixup prop position in RunAfterParent

diff --git a/Engine/Plugins/Experimental/Animation/RelativeIKOp/Source/RelativeIKOp/Private/RelativeIKOp.cpp b/Engine/Plugins/Experimental/Animation/RelativeIKOp/Source/RelativeIKOp/Private/RelativeIKOp.cpp
--- a/Engine/Plugins/Experimental/Animation/RelativeIKOp/Source/RelativeIKOp/Private/RelativeIKOp.cpp
+++ b/Engine/Plugins/Experimental/Animation/RelativeIKOp/Source/RelativeIKOp/Private/RelativeIKOp.cpp
@@ -2,7 +2,6 @@
 
 #include "RelativeIKOp.h"
 
-#include "IKRigObjectVersion.h"
 #include "IKRigDebugRendering.h"
 #include "RelativeBodyAnimNotifies.h"
 #include "Animation/AnimInstance.h"
@@ -13,7 +12,6 @@
 #include "PhysicsEngine/SkeletalBodySetup.h"
 #include "Retargeter/IKRetargetProcessor.h"
 #include "Retargeter/RetargetOps/RunIKRigOp.h"
-#include "Rig/Solvers/PointsToRotation.h"
 
 #if WITH_EDITOR
 #include "PrimitiveDrawingUtils.h"
@@ -99,42 +97,151 @@ void FIKRetargetRelativeIKOp::Run(
 		return;
 	}
 	
-	if (!bIsInitialized || !CacheSourceAnimSequence || !CachedNotifyInfo || AnimSeqPlayHead < 0.0f)
+	if (!bIsInitialized || !CacheSourceAnimSequence || AnimSeqPlayHead < 0.0f)
 	{
 		return;
 	}
 	
-	TConstArrayView<FVector3f> FramePointView = ApplyTemporalSmoothing(AnimSeqPlayHead);
-
+	// TODO: Figure out how to move these out of processing functions?
 	FDebugRelativeIKDrawInfo LocalDebugInfo;
+	TArray<FDebugRelativeTargetPairSpace> LocalTargetPairSpaces;
+
+	int32 NumMeshPairs = (CachedNotifyInfo) ? CachedNotifyInfo->bBodyPairsIsParentDominates.Num() : 0;
+	int32 NumPropPairs = (CachedPropNotifyInfo) ? (CachedPropNotifyInfo->PropsPairs.Num() / 2) : 0;
+	int32 TotalPairs = NumMeshPairs + NumPropPairs;
 	
-#if WITH_EDITOR
-	LocalDebugInfo.SourcePairVerts.Reserve(CachedNotifyInfo->bBodyPairsIsParentDominates.Num());
-	LocalDebugInfo.TargetPairVerts.Reserve(CachedNotifyInfo->bBodyPairsIsParentDominates.Num());
+	int32 NumMeshBodies = (CachedNotifyInfo) ? CachedNotifyInfo->BodyPairs.Num() : 0;
+	int32 NumPropBones = (CachedPropNotifyInfo) ? CachedPropNotifyInfo->PropsPairs.Num() : 0;
+	int32 TotalBodies = NumMeshBodies + NumPropBones;
 
-	TArray<FDebugRelativeTargetPairSpace> LocalTargetPairSpaces;
-	LocalTargetPairSpaces.SetNumZeroed(CachedNotifyInfo->BodyPairs.Num());
+	if ( TotalBodies == 0 )
+	{
+		return;
+	}
+	
+#if WITH_EDITOR
+	LocalDebugInfo.SourcePairVerts.Reserve(TotalPairs);
+	LocalDebugInfo.TargetPairVerts.Reserve(TotalPairs);
+	
+	LocalTargetPairSpaces.SetNumZeroed(TotalBodies);
 #endif //WITH_EDITOR
 
+	// Copy (pin) prop bones in target
+	// TODO: Pin in local space instead of comp space?
+	for (FName SourcePropBone : CachePinSourcePropBones)
+	{
+		FName TargetPropBone = ApplyBodyMap(SourcePropBone);
+		if (!CacheSourceSkelIndices.Contains(SourcePropBone) || !CacheTargetSkelIndices.Contains(TargetPropBone))
+		{
+			continue;
+		}
+		
+		int32 SourceBoneIdx = CacheSourceSkelIndices[SourcePropBone];
+		int32 TargetBoneIdx = CacheTargetSkelIndices[TargetPropBone];
+		OutTargetGlobalPose[TargetBoneIdx] = InSourceGlobalPose[SourceBoneIdx];
+	}
+
 	const double SourceScale = InProcessor.GetSourceScaleFactor();
-	const double InvSourceScale = (Settings.bIgnoreSourceScale) ? 1.0 / InProcessor.GetSourceScaleFactor() : 1.0;
+	if (CachedNotifyInfo)
+	{
+		int32 NumMeshSamples = CachedNotifyInfo->NumSamples;
+		TConstArrayView<FVector3f> FramePointView = ApplyTemporalSmoothing(AnimSeqPlayHead, MeshSampleRate, NumMeshSamples, NumMeshBodies, CachedNotifyInfo->BodyPairsLocalReference);
+		ComputeFramePairSpaces(SourceScale, InSourceGlobalPose, OutTargetGlobalPose, CacheMeshPairOffset, FramePointView, LocalDebugInfo, LocalTargetPairSpaces);
+	}
+
+	if (CachedPropNotifyInfo)
+	{
+		int32 NumPropSamples = CachedPropNotifyInfo->NumSamples;
+		TConstArrayView<FVector3f> FramePropView = ApplyTemporalSmoothing(AnimSeqPlayHead, PropSampleRate, NumPropSamples, NumPropBones, CachedPropNotifyInfo->PropsPairsLocalReference);
+		ComputeFramePropSpaces(SourceScale, InSourceGlobalPose, OutTargetGlobalPose, CachePropPairOffset, FramePropView, LocalDebugInfo, LocalTargetPairSpaces);
+	}
 	
 	FIKRigGoalContainer& GoalContainer = InProcessor.GetIKRigGoalContainer();
 	TArray<FIKRigGoal>& IKGoals = GoalContainer.GetGoalArray();
+	
+	LocalDebugInfo.TargetGoals.Reserve(CacheSourceBodyEffectVertIdx.Num());
 
+	CachePinTargetPropTransforms.SetNumZeroed(CachePinSourcePropBones.Num());
+	// Update bone goals to relative IK positions
+	for (int EffectIdx = 0; EffectIdx < CacheSourceBodyEffectVertIdx.Num(); ++EffectIdx)
+	{
+		TArray<int32>& EffectIndices = CacheSourceBodyEffectVertIdx[EffectIdx];
+		if (EffectIndices.IsEmpty())
+		{
+			continue;
+		}
+
+		FName SourceBoneName = CacheSourceEffectBones[EffectIdx];
+		FName TargetBoneName = ApplyBodyMap(SourceBoneName);
+		if (!CacheTargetSkelIndices.Contains(TargetBoneName))
+		{
+			continue;
+		}
+
+		int32 TargetBoneIdx = CacheTargetSkelIndices[TargetBoneName];
+		const FTransform& TargetBonePose = OutTargetGlobalPose[TargetBoneIdx];
+		UpdateTargetIKGoal(IKGoals, TargetBoneName, TargetBonePose, EffectIndices, LocalDebugInfo, LocalTargetPairSpaces);
+		
+		UpdatePropPinTransform(SourceBoneName, TargetBonePose, EffectIndices, LocalDebugInfo, LocalTargetPairSpaces);
+	}
+
+#if WITH_EDITOR
+	UpdateDebugInfo(SourceScale, InSourceGlobalPose, OutTargetGlobalPose, LocalDebugInfo);
+#endif
+}
+
+void FIKRetargetRelativeIKOp::RunAfterParent(
+	FIKRetargetProcessor& InProcessor,
+	const double InDeltaTime,
+	const TArray<FTransform>& InSourceGlobalPose,
+	TArray<FTransform>& OutTargetGlobalPose)
+{
+	if (!Settings.bApplyPostSolvePinning)
+	{
+		return;
+	}
+		
+	for (int32 Index=0; Index < CachePinSourcePropBones.Num(); ++Index)
+	{
+		FName SourcePinBone = CachePinSourcePropBones[Index]; 
+		FName TargetPinBone = ApplyBodyMap(SourcePinBone);
+		if (TargetPinBone == NAME_None || !CacheTargetSkelIndices.Contains(TargetPinBone))
+		{
+			continue;
+		}
+
+		int32 TargetBoneIdx = CacheTargetSkelIndices[TargetPinBone];
+		OutTargetGlobalPose[TargetBoneIdx] = CachePinTargetPropTransforms[Index];
+	}
+}
+
+void FIKRetargetRelativeIKOp::ComputeFramePairSpaces(
+	double SourceScale,
+	const TArray<FTransform>& SourcePose,
+	const TArray<FTransform>& TargetPose,
+	int32 PairOffset,
+	TConstArrayView<FVector3f> FramePointView,
+	FDebugRelativeIKDrawInfo& DebugInfo,
+	TArray<FDebugRelativeTargetPairSpace>& DebugPairSpaces)
+{
+	int32 PairCount = CachedNotifyInfo->bBodyPairsIsParentDominates.Num();
+	const double InvSourceScale = (Settings.bIgnoreSourceScale) ? 1.0 / SourceScale : 1.0;
+
+	// TODO: Should we just use initial chain lengths?
 	CacheChainLengthMap.Reset();
 	CacheChainStartMap.Reset();
 	
 	const double RetargetContactAlpha = FMath::Clamp(Settings.RetargetContactAlpha, 0.0, 1.0);
 	const double RetargetSpringAlpha = FMath::Clamp(Settings.RetargetSpringAlpha, 0.0, 1.0);
 	// Get pair target body verts and weights
-	for (int PairIdx = 0; PairIdx < CachedNotifyInfo->bBodyPairsIsParentDominates.Num(); ++PairIdx)
+	for (int PairIdx = 0; PairIdx < PairCount; ++PairIdx)
 	{
 		// TODO: Allow selecting a subset of body pairs in retargeter
-		FrameBoneWeights[2*PairIdx] = 0;
-		FrameBoneWeights[2*PairIdx + 1] = 0;
-		FrameBoneTargets[2*PairIdx] = FVector::ZeroVector;
-		FrameBoneTargets[2*PairIdx + 1] = FVector::ZeroVector;
+		const int32 OutPairIdx = PairIdx + PairOffset;
+		FrameBoneWeights[2*OutPairIdx] = 0;
+		FrameBoneWeights[2*OutPairIdx + 1] = 0;
+		FrameBoneTargets[2*OutPairIdx] = FVector::ZeroVector;
+		FrameBoneTargets[2*OutPairIdx + 1] = FVector::ZeroVector;
 		
 		FName SourcePrimaryBone = CachedNotifyInfo->BodyPairs[2*PairIdx];
 		FName SourceOtherBone = CachedNotifyInfo->BodyPairs[2*PairIdx + 1];
@@ -176,25 +283,24 @@ void FIKRetargetRelativeIKOp::Run(
 		const FVector3f& LocalOtherVert = FramePointView[2*PairIdx+1];
 
 		// TODO: Clean this up if we bake in rot invariant space
-		FVector SourceBoneBakePrimaryVert = (SourcePrimaryBodyRot).TransformPosition(FVector(LocalPrimaryVert)); 
-		FVector SourceBoneBakeOtherVert = (SourceOtherBodyRot).TransformPosition(FVector(LocalOtherVert));
+		FVector SourceBoneBakePrimaryV = (SourcePrimaryBodyRot).TransformPosition(FVector(LocalPrimaryVert)); 
+		FVector SourceBoneBakeOtherV = (SourceOtherBodyRot).TransformPosition(FVector(LocalOtherVert));
 
 		FBodyTransform SourcePrimaryTfm;
 		FBodyTransform SourceOtherTfm;
+		ComputeSourceBodyTransform(SourcePrimaryTfm, SourcePrimaryBone, SourcePose[SourcePrimaryPoseIdx], SourceScale);
+		ComputeSourceBodyTransform(SourceOtherTfm, SourceOtherBone, SourcePose[SourceOtherPoseIdx], SourceScale);
 
-		ComputeSourceBodyTransform(SourcePrimaryTfm, SourcePrimaryBone, InSourceGlobalPose[SourcePrimaryPoseIdx], SourceScale);
-		ComputeSourceBodyTransform(SourceOtherTfm, SourceOtherBone, InSourceGlobalPose[SourceOtherPoseIdx], SourceScale);
-
-		FVector SourcePrimaryCompV = ApplyBodyTransform(SourcePrimaryTfm, SourceBoneBakePrimaryVert);
-		FVector SourceOtherCompV = ApplyBodyTransform(SourceOtherTfm, SourceBoneBakeOtherVert);
+		FVector SourcePrimaryCompV = ApplyBodyTransform(SourcePrimaryTfm, SourceBoneBakePrimaryV);
+		FVector SourceOtherCompV = ApplyBodyTransform(SourceOtherTfm, SourceBoneBakeOtherV);
 
 		// Rescale components for distance checks if "Ignore sources scale" is on
 		// NOTE: This works because we don't need absolute coordinates (just distance).
 		//       If absolute coords needed then must properly account for scale pivot
 		FVector RescaleSourcePrimaryV = InvSourceScale * SourcePrimaryCompV;
-		FVector RecsaleSourceOtherV = InvSourceScale * SourceOtherCompV;
+		FVector RescaleSourceOtherV = InvSourceScale * SourceOtherCompV;
 		
-		double SourceVDist = FVector::Distance(RescaleSourcePrimaryV, RecsaleSourceOtherV);
+		double SourceVDist = FVector::Distance(RescaleSourcePrimaryV, RescaleSourceOtherV);
 		double LinearWeight = FMath::Clamp(1.0 - (SourceVDist / Settings.DistanceThreshold), 0.0, 1.0);
 		double Weight = GetDistanceWeight(SourceVDist, Settings.DistanceThreshold, Settings.DistanceFade);
 
@@ -209,21 +315,20 @@ void FIKRetargetRelativeIKOp::Run(
 
 		FBodyTransform TargetPrimaryTfm;
 		FBodyTransform TargetOtherTfm;
-
-		ComputeTargetBodyTransform(TargetPrimaryTfm, TargetPrimaryBone, OutTargetGlobalPose[TargetPrimaryPoseIdx]);
-		ComputeTargetBodyTransform(TargetOtherTfm, TargetOtherBone, OutTargetGlobalPose[TargetOtherPoseIdx]);
+		ComputeTargetBodyTransform(TargetPrimaryTfm, TargetPrimaryBone, TargetPose[TargetPrimaryPoseIdx]);
+		ComputeTargetBodyTransform(TargetOtherTfm, TargetOtherBone, TargetPose[TargetOtherPoseIdx]);
 
 		// Compute "primary" body vert represented in component space relative to primary (_P) and other (_O) bodies
-		FVector TargetPrimaryCompV_P = ApplyBodyTransform(TargetPrimaryTfm, SourceBoneBakePrimaryVert);
+		FVector TargetPrimaryCompV_P = ApplyBodyTransform(TargetPrimaryTfm, SourceBoneBakePrimaryV);
 		FVector TargetPrimaryCompV_O = ApplyBodyTransform(TargetOtherTfm, SourcePrimaryV_O);
 
 		// Compute "other" body vert represented in component space relative to primary (_P) and other (_O) bodies
-		FVector TargetOtherCompV_O = ApplyBodyTransform(TargetOtherTfm, SourceBoneBakeOtherVert);
+		FVector TargetOtherCompV_O = ApplyBodyTransform(TargetOtherTfm, SourceBoneBakeOtherV);
 		FVector TargetOtherCompV_P = ApplyBodyTransform(TargetPrimaryTfm, SourceOtherV_P);
 
 		// Get offsets from target verts (in their own body rep) to associated bone
-		FVector TargetPrimaryBoneOffset = OutTargetGlobalPose[TargetPrimaryPoseIdx].GetTranslation() - TargetPrimaryCompV_P;
-		FVector TargetOtherBoneOffset = OutTargetGlobalPose[TargetOtherPoseIdx].GetTranslation() - TargetOtherCompV_O;
+		FVector TargetPrimaryBoneOffset = TargetPose[TargetPrimaryPoseIdx].GetTranslation() - TargetPrimaryCompV_P;
+		FVector TargetOtherBoneOffset = TargetPose[TargetOtherPoseIdx].GetTranslation() - TargetOtherCompV_O;
 
 		// double SpringAlpha = (Settings.bTestDistContactAlpha) ? Weight : RetargetSpringAlpha;
 		// These are the target midpoint positions used for spring constraints in PhysicalIK
@@ -240,7 +345,7 @@ void FIKRetargetRelativeIKOp::Run(
 		bool bDistanceWeighting = CachedNotifyInfo->bBodyPairsIsParentDominates[PairIdx] && Settings.bTestDistContactAlpha;
 		
 		FDoubleInterval FeasibleRange(0.0, 1.0);
-		UpdateCacheChainInfo(OutTargetGlobalPose, TargetOtherBone);
+		UpdateCacheChainInfo(TargetPose, TargetOtherBone);
 		if (bDistanceWeighting)
 		{
 			// TODO: Experiment with bi-directional distance weighting when Primary/Other can both move
@@ -264,192 +369,322 @@ void FIKRetargetRelativeIKOp::Run(
 		// Debug draw verts
 		if (Settings.bDebugDraw)
 		{
-			LocalDebugInfo.SourcePairVerts.Add({SourcePrimaryCompV, SourceOtherCompV, false, FVector::ZeroVector, FVector::ZeroVector});
+			DebugInfo.SourcePairVerts.Add({SourcePrimaryCompV, SourceOtherCompV, false, FVector::ZeroVector, FVector::ZeroVector});
 			if (Settings.DebugFullRetargetPairBones.Contains(TargetPrimaryBone))
 			{
-				LocalDebugInfo.TargetPairVerts.Add({TargetPrimaryCompV_P, TargetOtherCompV_O, true,TargetOtherCompV_P, TargetPrimaryCompV_O,});
+				DebugInfo.TargetPairVerts.Add({TargetPrimaryCompV_P, TargetOtherCompV_O, true,TargetOtherCompV_P, TargetPrimaryCompV_O,});
 			}
 			else
 			{
-				LocalDebugInfo.TargetPairVerts.Add({TargetPrimaryCompV_P, TargetOtherCompV_O, false, FVector::ZeroVector, FVector::ZeroVector});
+				DebugInfo.TargetPairVerts.Add({TargetPrimaryCompV_P, TargetOtherCompV_O, false, FVector::ZeroVector, FVector::ZeroVector});
 			}
 			
-			LocalTargetPairSpaces[2*PairIdx] = FDebugRelativeTargetPairSpace{WeightedTargetPrimary, TargetRelativeToOtherV+RelativePrimaryOffset, TargetRelativeToPrimaryV+RelativePrimaryOffset, FeasibleRange};
-			LocalTargetPairSpaces[2*PairIdx+1] = FDebugRelativeTargetPairSpace{WeightedTargetOther, TargetRelativeToOtherV+RelativeOtherOffset, TargetRelativeToPrimaryV+RelativeOtherOffset, FeasibleRange};
+			DebugPairSpaces[2*OutPairIdx] = FDebugRelativeTargetPairSpace{WeightedTargetPrimary, TargetRelativeToOtherV+RelativePrimaryOffset, TargetRelativeToPrimaryV+RelativePrimaryOffset, FeasibleRange};
+			DebugPairSpaces[2*OutPairIdx+1] = FDebugRelativeTargetPairSpace{WeightedTargetOther, TargetRelativeToOtherV+RelativeOtherOffset, TargetRelativeToPrimaryV+RelativeOtherOffset, FeasibleRange};
 		}
 #endif //WITH_EDITOR
 	}
-	
-	LocalDebugInfo.TargetGoals.Reserve(CacheSourceBodyEffectVertIdx.Num());
-	// Update bone goals to relative IK positions
-	for (int EffectIdx = 0; EffectIdx < CacheSourceBodyEffectVertIdx.Num(); ++ EffectIdx)
-	{
-		TArray<int32>& EffectIndices = CacheSourceBodyEffectVertIdx[EffectIdx];
-		if (EffectIndices.IsEmpty())
-		{
-			continue;
-		}
+}
 
-		FName SourceBoneName = CacheSourceEffectBones[EffectIdx];
-		FName TargetBoneName = ApplyBodyMap(SourceBoneName);
-		// HACK: Quick and dirty goal-by-bone find
-		FIKRigGoal* IKGoal = IKGoals.FindByPredicate([TargetBoneName](const FIKRigGoal& Other){return Other.BoneName == TargetBoneName;});
-		if (!IKGoal)
-		{
-			continue;
-		}
+void FIKRetargetRelativeIKOp::ComputeFramePropSpaces(
+	double SourceScale,
+	const TArray<FTransform>& SourcePose,
+	const TArray<FTransform>& TargetPose,
+	int32 PairOffset,
+	TConstArrayView<FVector3f> FramePointView,
+	FDebugRelativeIKDrawInfo& DebugInfo,
+	TArray<FDebugRelativeTargetPairSpace>& DebugPairSpaces)
+{
+	int32 PropCount = CachedPropNotifyInfo->PropsOffsetTransforms.Num();
+	int32 BodyCount = CachedPropNotifyInfo->PropsPairs.Num() / PropCount / 2;
 
-		if (!CacheTargetSkelIndices.Contains(TargetBoneName))
+	FTransform SourceScaleTfm(FQuat::Identity, FVector::ZeroVector, FVector(SourceScale));
+	const double InvSourceScale = (Settings.bIgnoreSourceScale) ? 1.0 / SourceScale : 1.0;
+	
+	// CacheChainLengthMap.Reset();
+	// CacheChainStartMap.Reset();
+	
+	const double RetargetContactAlpha = FMath::Clamp(Settings.RetargetContactAlpha, 0.0, 1.0);
+	const double RetargetSpringAlpha = FMath::Clamp(Settings.RetargetSpringAlpha, 0.0, 1.0);
+	// Get pair target body verts and weights
+	for (int PropIdx = 0; PropIdx < PropCount; ++PropIdx)
+	{
+		const FTransform PropTfm = CachedPropNotifyInfo->PropsOffsetTransforms[PropIdx];
+		for (int BodyIdx = 0; BodyIdx < BodyCount; ++BodyIdx)
 		{
-			continue;
-		}
+			const int32 PairIdx = BodyCount * PropIdx + BodyIdx;
+			const int32 OutPairIdx = PairIdx + PairOffset;
+			
+			FrameBoneWeights[2*OutPairIdx] = 0;
+			FrameBoneWeights[2*OutPairIdx + 1] = 0;
+			FrameBoneTargets[2*OutPairIdx] = FVector::ZeroVector;
+			FrameBoneTargets[2*OutPairIdx + 1] = FVector::ZeroVector;
+			
+			FName SourcePropBone = CachedPropNotifyInfo->PropsPairs[2*PairIdx];
+			FName SourceBodyBone = CachedPropNotifyInfo->PropsPairs[2*PairIdx + 1];
 
-		double TotalWeight = 0;
-		FVector GoalPoint = FVector::ZeroVector;
+			FName TargetPropBone = ApplyBodyMap(SourcePropBone);
+			FName TargetBodyBone = ApplyBodyMap(SourceBodyBone);
 
-#if WITH_EDITOR
-		if (Settings.bDebugDraw)
-		{
-			LocalDebugInfo.TargetGoals.AddDefaulted();
-			LocalDebugInfo.TargetGoals.Last().Goal = FVector::ZeroVector;
-		}
-#endif //WITH_EDITOR
+			if (TargetPropBone == NAME_None || TargetBodyBone == NAME_None)
+			{
+				continue;
+			}
 
-		double MaxWeight = 0;
-		for (int Index : EffectIndices)
-		{
-			if (!FMath::IsNearlyZero(FrameBoneWeights[Index]))
+			if (!CacheSourceSkelIndices.Contains(SourcePropBone) || !CacheSourceSkelIndices.Contains(SourceBodyBone))
 			{
-				GoalPoint += FrameBoneWeights[Index] * FrameBoneTargets[Index];
-				TotalWeight += FrameBoneWeights[Index];
+				continue;
+			}
 
-#if WITH_EDITOR
-				if (Settings.bDebugDraw)
-				{
-					LocalDebugInfo.TargetGoals.Last().PairTargets.Add(FrameBoneTargets[Index]);
-					LocalDebugInfo.PairRetargetInfo.Emplace(LocalTargetPairSpaces[Index]);
-				}
-#endif //WITH_EDITOR
+			if (!CacheTargetSkelIndices.Contains(TargetPropBone) || !CacheTargetSkelIndices.Contains(TargetBodyBone))
+			{
+				continue;
 			}
-		}
 
-		
-		if (!FMath::IsNearlyZero(TotalWeight))
-		{
-			GoalPoint /= TotalWeight;
-#if WITH_EDITOR
-			if (Settings.bDebugDraw)
+			if (!HasBoneDelta(TargetPropBone) || !HasBoneDelta(TargetBodyBone))
 			{
-				LocalDebugInfo.TargetGoals.Last().Goal = GoalPoint;
+				continue;
 			}
-#endif // WITH_EDITOR
+
+			int32 SourcePropPoseIdx = CacheSourceSkelIndices[SourcePropBone];
+			int32 SourceBodyPoseIdx = CacheSourceSkelIndices[SourceBodyBone];
+
+			// ---- TESTING ONLY - Remove these once notify baking is rotation invariant
+			FTransform SourceBodyRot = FTransform(GetBodyRotation(SourcePhysicsAsset, SourceBodyBone));
+			FTransform TargetBodyRot = SourceBodyRot;
+			
+			// Compute source transforms and source-space vertex distance
+			// const FVector3f& LocalPrimaryVert = CachedNotifyInfo->BodyPairsLocalReference[FramePairIdx + 2*PairIdx];
+			// const FVector3f& LocalOtherVert = CachedNotifyInfo->BodyPairsLocalReference[FramePairIdx + 2*PairIdx + 1];
+			const FVector3f& LocalPropVert = FramePointView[2*PairIdx];
+			const FVector3f& LocalBodyVert = FramePointView[2*PairIdx+1];
+
+			// TODO: Clean this up if we bake in rot invariant space
+			FVector SourceBoneBakePropV = (PropTfm).TransformPosition(FVector(LocalPropVert));
+			FVector SourceBoneBakeBodyV = (SourceBodyRot).TransformPosition(FVector(LocalBodyVert));
+			
+			FTransform SourcePropTfm = SourceScaleTfm * SourcePose[SourcePropPoseIdx];
+			FBodyTransform SourceBodyTfm;
+			ComputeSourceBodyTransform(SourceBodyTfm, SourceBodyBone, SourcePose[SourceBodyPoseIdx], SourceScale);
+			
+			FVector SourcePropCompV = SourcePropTfm.TransformPosition(SourceBoneBakePropV);
+			FVector SourceBodyCompV = ApplyBodyTransform(SourceBodyTfm, SourceBoneBakeBodyV);
+
+			// Rescale components for distance checks if "Ignore sources scale" is on
+			// NOTE: This works because we don't need absolute coordinates (just distance).
+			//       If absolute coords needed then must properly account for scale pivot
+			FVector RescaleSourcePropV = InvSourceScale * SourcePropCompV;
+			FVector RescaleSourceBodyV = InvSourceScale * SourceBodyCompV;
+
+			double SourceVDist = FVector::Distance(RescaleSourcePropV, RescaleSourceBodyV);
+			double LinearWeight = FMath::Clamp(1.0 - (SourceVDist / Settings.DistanceThreshold), 0.0, 1.0);
+			double Weight = GetDistanceWeight(SourceVDist, Settings.DistanceThreshold, Settings.DistanceFade);
+
+			// Compute Source pair verts as represented in alternate pair body primary(_P) or other(_O) local bake space
+			FVector SourcePropV_B = InverseBodyTransform(SourceBodyTfm, SourcePropCompV);
+			FVector SourceBodyV_P = SourcePropTfm.InverseTransformPosition(SourceBodyCompV);
 			
-			// TODO: Play with total weight functions to alpha goal
-			float GoalAlpha = FMath::Min(1.0f, static_cast<float>(TotalWeight*FMath::Clamp(Settings.ContributionSumWeight, 0.0f, 1.0f)));
-			int32 TargetBoneIdx = CacheTargetSkelIndices[TargetBoneName];
-			FVector PrevGoal = GoalLocBlendCompSpace(IKGoal, OutTargetGlobalPose[TargetBoneIdx]);
-			FVector BlendGoal = FMath::Lerp(PrevGoal, GoalPoint, GoalAlpha);
-
-			// TODO: Remove dry run for release
-			bool bSetGoals = !Settings.bDebugDraw || !Settings.bDryRun;
-			if (bSetGoals)
+			//---------------------------------
+			// Get verts retargeted using physics body oriented scale
+			int32 TargetPropPoseIdx = CacheTargetSkelIndices[TargetPropBone];
+			int32 TargetBodyPoseIdx = CacheTargetSkelIndices[TargetBodyBone];
+
+			FTransform TargetPropBoneDelta = GetTargetBoneDelta(TargetPropBone);
+
+			FTransform TargetPropTfm = TargetPropBoneDelta * TargetPose[TargetPropPoseIdx];
+			FBodyTransform TargetBodyTfm;
+			ComputeTargetBodyTransform(TargetBodyTfm, TargetBodyBone, TargetPose[TargetBodyPoseIdx]);
+
+			// Compute prop vert represented in component space relative to prop (_P) and body (_B)
+			FVector TargetPropCompV_P = TargetPropTfm.TransformPosition(SourceBoneBakePropV);
+			FVector TargetPropCompV_B = ApplyBodyTransform(TargetBodyTfm, SourcePropV_B);
+
+			// Compute body vert represented in component space relative to prop (_P) and body (_B)
+			FVector TargetBodyCompV_B = ApplyBodyTransform(TargetBodyTfm, SourceBoneBakeBodyV);
+			FVector TargetBodyCompV_P = TargetPropTfm.TransformPosition(SourceBodyV_P);
+
+			// Get offsets from target verts (in their own body rep) to associated bone
+			FVector TargetPropBoneOffset = TargetPose[TargetPropPoseIdx].GetTranslation() - TargetPropCompV_P;
+			FVector TargetBodyBoneOffset = TargetPose[TargetBodyPoseIdx].GetTranslation() - TargetBodyCompV_B;
+
+			// double SpringAlpha = (Settings.bTestDistContactAlpha) ? Weight : RetargetSpringAlpha;
+			FVector TargetRelativeToPropV = FMath::Lerp(TargetPropCompV_P,TargetBodyCompV_P, RetargetSpringAlpha);
+			FVector TargetRelativeToBodyV = FMath::Lerp(TargetPropCompV_B,TargetBodyCompV_B, RetargetSpringAlpha);
+
+			// double TargetVDist = FVector::Distance(TargetRelativeToPrimaryV, TargetRelativeToOtherV);
+			// double TargetDRat = FMath::Min(1.0, 2.0 * FMath::Min(SourceVDist/TargetVDist, TargetVDist/SourceVDist));
+			
+			FVector RelativePropOffset = TargetPropCompV_P - TargetRelativeToPropV;
+			FVector RelativeBodyOffset = TargetBodyCompV_B - TargetRelativeToBodyV;
+			
+			// bool bDistanceWeighting = false;
+			
+			FDoubleInterval FeasibleRange(0.0, 1.0);
+			// UpdateCacheChainInfo(TargetPose, TargetBodyBone);
+			// if (bDistanceWeighting)
+			// {
+			// 	ComputeFeasibilityRange(FeasibleRange, TargetBodyBone, RelativeBodyOffset + TargetBodyBoneOffset, TargetRelativeToBodyV, TargetRelativeToPropV);
+			// }
+
+			double FeasibilityWeight = FMath::Lerp(FeasibleRange.Min, FeasibleRange.Max, LinearWeight);
+			// double ContactAlpha = (bDistanceWeighting) ? FeasibilityWeight : RetargetContactAlpha;
+			double ContactAlpha = RetargetContactAlpha;
+			FVector WeightedTarget = FMath::Lerp(TargetRelativeToBodyV, TargetRelativeToPropV, ContactAlpha);
+
+			FVector WeightedTargetProp = WeightedTarget + RelativePropOffset;
+			FVector WeightedTargetBody = WeightedTarget + RelativeBodyOffset;
+			
+			FrameBoneTargets[2*OutPairIdx] = WeightedTargetProp + TargetPropBoneOffset;
+			FrameBoneTargets[2*OutPairIdx + 1] = WeightedTargetBody + TargetBodyBoneOffset;
+
+			FrameBoneWeights[2*OutPairIdx] = Weight;
+			FrameBoneWeights[2*OutPairIdx + 1] = Weight;
+
+#if WITH_EDITOR
+			// Debug draw verts
+			if (Settings.bDebugDraw)
 			{
-				IKGoal->bEnabled = true;
-				IKGoal->Position = BlendGoal;
-				IKGoal->PositionAlpha = 1.0;
-				IKGoal->PositionSpace = EIKRigGoalSpace::Component;
+				DebugInfo.SourcePairVerts.Add({SourcePropCompV, SourceBodyCompV, false, FVector::ZeroVector, FVector::ZeroVector});
+				if (Settings.DebugFullRetargetPairBones.Contains(TargetPropBone))
+				{
+					DebugInfo.TargetPairVerts.Add({TargetPropCompV_P, TargetBodyCompV_B, true,TargetBodyCompV_P, TargetPropCompV_B,});
+				}
+				else
+				{
+					DebugInfo.TargetPairVerts.Add({TargetPropCompV_P, TargetBodyCompV_B, false, FVector::ZeroVector, FVector::ZeroVector});
+				}
+				
+				DebugPairSpaces[2*OutPairIdx] = FDebugRelativeTargetPairSpace{WeightedTargetProp, TargetRelativeToBodyV+RelativePropOffset, TargetRelativeToPropV+RelativePropOffset, FeasibleRange};
+				DebugPairSpaces[2*OutPairIdx+1] = FDebugRelativeTargetPairSpace{WeightedTargetBody, TargetRelativeToBodyV+RelativeBodyOffset, TargetRelativeToPropV+RelativeBodyOffset, FeasibleRange};
 			}
+#endif //WITH_EDITOR
 		}
 	}
+}
 
-#if WITH_EDITOR
-	if (Settings.bDebugDraw)
+void FIKRetargetRelativeIKOp::UpdateTargetIKGoal(
+	TArray<FIKRigGoal>& OutIKGoals,
+	FName TargetBoneName,
+	const FTransform& TargetBonePose,
+	const TArray<int32>& EffectIndices,
+	FDebugRelativeIKDrawInfo& DebugInfo,
+	TArray<FDebugRelativeTargetPairSpace>& DebugPairSpaces)
+{
+	FIKRigGoal* IKGoal = OutIKGoals.FindByPredicate([TargetBoneName](const FIKRigGoal& Other){return Other.BoneName == TargetBoneName;});
+	if (!IKGoal)
 	{
-		FScopeLock ScopeLock(&DebugDataMutex);
+		return;
+	}
 
-		DebugDrawInfo.TargetGoals = LocalDebugInfo.TargetGoals;
-		DebugDrawInfo.PairRetargetInfo = LocalDebugInfo.PairRetargetInfo;
-		DebugDrawInfo.SourcePairVerts = LocalDebugInfo.SourcePairVerts;
-		DebugDrawInfo.TargetPairVerts = LocalDebugInfo.TargetPairVerts;
+	FVector TargetGoalLoc;
+	double TotalWeight = ComputeTargetWeightedLocation(TargetGoalLoc, EffectIndices, DebugInfo, DebugPairSpaces);
+	if (FMath::IsNearlyZero(TotalWeight))
+	{
+		return;
+	}
+		
+	// TODO: Play with total weight functions to alpha goal
+	double GoalAlpha = FMath::Min(1.0, TotalWeight*FMath::Clamp(Settings.ContributionSumWeight, 0.0f, 1.0f));
+	FVector PrevGoal = GoalLocBlendCompSpace(IKGoal, TargetBonePose);
+	FVector BlendGoal = FMath::Lerp(PrevGoal, TargetGoalLoc, GoalAlpha);
+
+	// TODO: Remove dry run for release
+	bool bSetGoals = !Settings.bDebugDraw || !Settings.bDryRun;
+	if (bSetGoals)
+	{
+		IKGoal->bEnabled = true;
+		IKGoal->Position = BlendGoal;
+		IKGoal->PositionAlpha = 1.0;
+		IKGoal->PositionSpace = EIKRigGoalSpace::Component;
+	}
+}
 
-		// Cache debug info for source bodies
-		DebugDrawInfo.SourceBodyInfo.Reset(CacheSourceEffectBones.Num());
-		for (FName SourceBoneName : CacheSourceEffectBones)
-		{
-			if (!CacheSourceSkelIndices.Contains(SourceBoneName))
-			{
-				continue;
-			}
 
-			int32 SourcePoseIdx = CacheSourceSkelIndices[SourceBoneName];
-			// Need to apply source-scale to bodies since it's directly applied to translations only in source pose 
-			FTransform SourceScaleTfm{FQuat::Identity, FVector::ZeroVector, FVector(SourceScale)};
-			DebugDrawInfo.SourceBodyInfo.Emplace(SourceBoneName, SourceScale, SourceScaleTfm*InSourceGlobalPose[SourcePoseIdx], FTransform::Identity);
-		}
+void FIKRetargetRelativeIKOp::UpdatePropPinTransform(
+	FName SourceBoneName,
+	const FTransform& TargetBonePose,
+	const TArray<int32>& EffectIndices,
+	FDebugRelativeIKDrawInfo& DebugInfo,
+	TArray<FDebugRelativeTargetPairSpace>& DebugPairSpaces)
+{
+	int32 PinPropIdx = CachePinSourcePropBones.Find(SourceBoneName);
+	if (PinPropIdx == INDEX_NONE)
+	{
+		return;
+	}
 
-		// Cache debug info for target domain bodies
-		DebugDrawInfo.TargetBodyInfo.Reset(CacheSourceDomainBones.Num());
-		for (FName SourceBoneName : CacheSourceDomainBones)
-		{
-			FName TargetBoneName = ApplyBodyMap(SourceBoneName);
-			if (!CacheTargetSkelIndices.Contains(TargetBoneName))
-			{
-				continue;
-			}
+	// Copy comp-space pose to re-apply regardless
+	CachePinTargetPropTransforms[PinPropIdx] = TargetBonePose;
 
-			FTransform TargetBoneDelta = GetTargetBoneDelta(TargetBoneName);
-			int32 TargetPoseIdx = CacheTargetSkelIndices[TargetBoneName];
-			DebugDrawInfo.TargetBodyInfo.Emplace(TargetBoneName, 1.0, OutTargetGlobalPose[TargetPoseIdx], TargetBoneDelta);
-		}
+	FVector TargetPropBoneLoc;
+	double TotalWeight = ComputeTargetWeightedLocation(TargetPropBoneLoc, EffectIndices, DebugInfo, DebugPairSpaces);
+	if (FMath::IsNearlyZero(TotalWeight))
+	{
+		return;
+	}
+	
+	double PropBoneAlpha = FMath::Min(1.0, TotalWeight*FMath::Clamp(Settings.ContributionSumWeight, 0.0, 1.0));
+	double BlendAlpha = PropBoneAlpha * (1.0-Settings.BlendPinsToSource);
+	FVector BlendBoneLoc = FMath::Lerp(TargetBonePose.GetLocation(), TargetPropBoneLoc, BlendAlpha);
+	
+	CachePinTargetPropTransforms[PinPropIdx].SetLocation(BlendBoneLoc);
+}
 
-		DebugDrawInfo.SourceTfmInfo.Reset(CacheSourceEffectBones.Num());
-		for (FName SourceBoneName : CacheSourceEffectBones)
-		{
-			if (!CacheSourceSkelIndices.Contains(SourceBoneName))
-			{
-				continue;
-			}
+double FIKRetargetRelativeIKOp::ComputeTargetWeightedLocation(
+	FVector& OutTargetLocation,
+	const TArray<int32>& EffectIndices,
+	FDebugRelativeIKDrawInfo& DebugInfo,
+	TArray<FDebugRelativeTargetPairSpace>& DebugPairSpaces)
+{
+	double TotalWeight = 0;
+	OutTargetLocation = FVector::ZeroVector;
 
-			int32 SourcePoseIdx = CacheSourceSkelIndices[SourceBoneName];
-			
-			FBodyTransform SourceFullTfm;
-			FTransform SourceBodyRot = FTransform(GetBodyRotation(SourcePhysicsAsset, SourceBoneName));
-			ComputeSourceBodyTransform(SourceFullTfm, SourceBoneName, InSourceGlobalPose[SourcePoseIdx], SourceScale);
-
-			FVector SourceBodyCenter = ApplyBodyTransform(SourceFullTfm, FVector::ZeroVector);
-			FVector SourceBodyX = ApplyBodyTransform(SourceFullTfm, SourceBodyRot.GetUnitAxis(EAxis::X));
-			FVector SourceBodyY = ApplyBodyTransform(SourceFullTfm, SourceBodyRot.GetUnitAxis(EAxis::Y));
-			FVector SourceBodyZ = ApplyBodyTransform(SourceFullTfm, SourceBodyRot.GetUnitAxis(EAxis::Z));
-			
-			DebugDrawInfo.SourceTfmInfo.Emplace(SourceBodyCenter, SourceBodyX, SourceBodyY, SourceBodyZ);
+#if WITH_EDITOR
+	if (Settings.bDebugDraw)
+	{
+		DebugInfo.TargetGoals.AddDefaulted();
+		DebugInfo.TargetGoals.Last().Goal = FVector::ZeroVector;
+	}
+#endif //WITH_EDITOR
+
+	double MaxWeight = 0;
+	for (int Index : EffectIndices)
+	{
+		if (FMath::IsNearlyZero(FrameBoneWeights[Index]))
+		{
+			continue;
 		}
 		
-		DebugDrawInfo.TargetTfmInfo.Reset(CacheSourceEffectBones.Num());
-		for (FName SourceBoneName : CacheSourceEffectBones)
-		{
-			FName TargetBoneName = ApplyBodyMap(SourceBoneName);
-			if (!CacheTargetSkelIndices.Contains(TargetBoneName))
-			{
-				continue;
-			}
-			
-			int32 TargetPoseIdx = CacheTargetSkelIndices[TargetBoneName];
-
-			FBodyTransform TargetFullTfm;
-			FTransform SourceBodyRot = FTransform(GetBodyRotation(SourcePhysicsAsset, SourceBoneName));
-			ComputeTargetBodyTransform(TargetFullTfm, SourceBoneName, OutTargetGlobalPose[TargetPoseIdx]);
+		OutTargetLocation += FrameBoneWeights[Index] * FrameBoneTargets[Index];
+		TotalWeight += FrameBoneWeights[Index];
 
-			FVector TargetBodyCenter = ApplyBodyTransform(TargetFullTfm, FVector::ZeroVector);
-			FVector TargetBodyX = ApplyBodyTransform(TargetFullTfm, SourceBodyRot.GetUnitAxis(EAxis::X));
-			FVector TargetBodyY = ApplyBodyTransform(TargetFullTfm, SourceBodyRot.GetUnitAxis(EAxis::Y));
-			FVector TargetBodyZ = ApplyBodyTransform(TargetFullTfm, SourceBodyRot.GetUnitAxis(EAxis::Z));
-			
-			DebugDrawInfo.TargetTfmInfo.Emplace(TargetBodyCenter, TargetBodyX, TargetBodyY, TargetBodyZ);
+#if WITH_EDITOR
+		if (Settings.bDebugDraw)
+		{
+			DebugInfo.TargetGoals.Last().PairTargets.Add(FrameBoneTargets[Index]);
+			DebugInfo.PairRetargetInfo.Emplace(DebugPairSpaces[Index]);
 		}
+#endif //WITH_EDITOR
 	}
-#endif
+		
+	if (FMath::IsNearlyZero(TotalWeight))
+	{
+		return TotalWeight;
+	}
+		
+	OutTargetLocation /= TotalWeight;
+#if WITH_EDITOR
+	if (Settings.bDebugDraw)
+	{
+		DebugInfo.TargetGoals.Last().Goal = OutTargetLocation;
+	}
+#endif // WITH_EDITOR
+	
+	return TotalWeight;
 }
 
+
+
 FName FIKRetargetRelativeIKOp::ApplyBodyMap(FName BodyName) const
 {
 	if (Settings.BodyMapping.IsEmpty() || !Settings.BodyMapping.Contains(BodyName))
@@ -542,22 +777,26 @@ void FIKRetargetRelativeIKOp::UpdateCacheChainInfo(const TArray<FTransform>& Tar
 	CacheChainStartMap.Emplace(BoneName, TargetPose[CacheBoneChains[BoneName]->BoneIndices[0]].GetTranslation());
 }
 
-TConstArrayView<FVector3f> FIKRetargetRelativeIKOp::ApplyTemporalSmoothing(float Time)
+TConstArrayView<FVector3f> FIKRetargetRelativeIKOp::ApplyTemporalSmoothing(float Time, float SampleRate, int32 NumSamples, int32 NumBodies, const TArray<FVector3f>& PairLocalVerts)
 {
-	int32 NumBodies = CachedNotifyInfo->BodyPairs.Num();
-	int32 SampleIdx = FMath::Clamp(static_cast<int32>(Time * SampleRate), 0, CachedNotifyInfo->NumSamples-1);
+	if (NumSamples <= 0)
+	{
+		return {};
+	}
+	
+	int32 SampleIdx = FMath::Clamp(static_cast<int32>(Time * SampleRate), 0, NumSamples-1);
 
 	// Testing temporal smoothing of local points
 	const int32 MinSmoothIdx = FMath::Max(SampleIdx - Settings.TemporalSmoothingRadius, 0);
-	const int32 MaxSmoothIdx = FMath::Min(SampleIdx + Settings.TemporalSmoothingRadius,CachedNotifyInfo->NumSamples-1) ;
+	const int32 MaxSmoothIdx = FMath::Min(SampleIdx + Settings.TemporalSmoothingRadius,NumSamples-1) ;
 	
-	if (Settings.TemporalSmoothingRadius <= 0 || CachedNotifyInfo->NumSamples <= 0)
+	if (Settings.TemporalSmoothingRadius <= 0)
 	{
-		return {CachedNotifyInfo->BodyPairsLocalReference.GetData() + NumBodies * SampleIdx, NumBodies};
+		return {PairLocalVerts.GetData() + NumBodies * SampleIdx, NumBodies};
 	}
 	
 	SmoothedPoints.SetNumUninitialized(NumBodies);
-	for (int BodyIdx = 0; BodyIdx < CachedNotifyInfo->BodyPairs.Num(); ++BodyIdx)
+	for (int BodyIdx = 0; BodyIdx < NumBodies; ++BodyIdx)
 	{
 		SmoothedPoints[BodyIdx] = FVector3f::Zero();
 		float SumWeight = 0.0f;
@@ -567,7 +806,7 @@ TConstArrayView<FVector3f> FIKRetargetRelativeIKOp::ApplyTemporalSmoothing(float
 			const float Sigma = 0.25f * static_cast<float>(Settings.TemporalSmoothingRadius);
 			const float SmoothDist = static_cast<float>(FrameIdx - SampleIdx);
 			const float SmoothWeight = FMath::Exp(-0.5f*(SmoothDist*SmoothDist) / (Sigma*Sigma));
-			SmoothedPoints[BodyIdx] += SmoothWeight * CachedNotifyInfo->BodyPairsLocalReference[NumBodies * FrameIdx + BodyIdx];
+			SmoothedPoints[BodyIdx] += SmoothWeight * PairLocalVerts[NumBodies * FrameIdx + BodyIdx];
 			SumWeight += SmoothWeight;
 		}
 		SmoothedPoints[BodyIdx] /= SumWeight;
@@ -750,7 +989,6 @@ FKShapeElem* FIKRetargetRelativeIKOp::FindBodyShape(const UPhysicsAsset* PhysAss
 	int32 BodyIdx = PhysAsset->FindBodyIndex(BoneName);
 	if (BodyIdx == INDEX_NONE)
 	{
-		UE_LOG(LogAnimation, Warning, TEXT("No body index found: %s"), *BoneName.ToString());
 		return nullptr;
 	}
 	return PhysAsset->SkeletalBodySetups[BodyIdx]->AggGeom.GetElement(0);
@@ -896,8 +1134,7 @@ void FIKRetargetRelativeIKOp::SetupRelativeIKNotifyInfoMontage()
 	{
 		return;
 	}
-
-	// HACK: Only grabs first relative notify
+	
 	for (const FSlotAnimationTrack& SlotTrack : Montage->SlotAnimTracks)
 	{
 		for (const FAnimSegment& Segment : SlotTrack.AnimTrack.AnimSegments)
@@ -908,17 +1145,24 @@ void FIKRetargetRelativeIKOp::SetupRelativeIKNotifyInfoMontage()
 			{
 				continue;
 			}
-			
+
+			// TODO: This loop only really works w/ at most one of each notify
 			for (const FAnimNotifyEvent& NotifyEvent : SeqRef->Notifies)
 			{
-				URelativeBodyBakeAnimNotify* BakeNotifyInfo = Cast<URelativeBodyBakeAnimNotify>(NotifyEvent.Notify);
-				if (BakeNotifyInfo)
+				if (URelativeBodyBakeAnimNotify* BodyInfo = Cast<URelativeBodyBakeAnimNotify>(NotifyEvent.Notify))
 				{
 					SegmentStartTime = Segment.AnimStartTime;
 					SegmentEndTime = Segment.AnimEndTime;
 					PreUpdateMontagePlayhead();
-					UpdateCacheNotifyInfo(BakeNotifyInfo);
-					return;
+					UpdateCacheNotifyInfo(BodyInfo);
+				}
+				
+				else if (URelativePropsBakeAnimNotify* PropInfo = Cast<URelativePropsBakeAnimNotify>(NotifyEvent.Notify))
+				{
+					SegmentStartTime = Segment.AnimStartTime;
+					SegmentEndTime = Segment.AnimEndTime;
+					PreUpdateMontagePlayhead();
+					UpdateCachePropNotifyInfo(PropInfo);
 				}
 			}
 		}
@@ -948,8 +1192,12 @@ void FIKRetargetRelativeIKOp::UpdateRelativeIKNotifyInfoAnimSeq(UAnimInstance* S
 					check(ActivePlayer.DeltaTimeRecord->IsPreviousValid())
 					for (const FAnimNotifyEvent& NotifyEvent : Sequence->Notifies)
 					{
-						URelativeBodyBakeAnimNotify* BakedNotifyInfo = Cast<URelativeBodyBakeAnimNotify>(NotifyEvent.Notify);
-						if (BakedNotifyInfo)
+						if (URelativeBodyBakeAnimNotify* BakedNotifyInfo = Cast<URelativeBodyBakeAnimNotify>(NotifyEvent.Notify))
+						{
+							MaxBlend = ActivePlayer.EffectiveBlendWeight;
+							MaxNotifyRecord = &ActivePlayer;
+						}
+						else if (URelativePropsBakeAnimNotify* PropInfo = Cast<URelativePropsBakeAnimNotify>(NotifyEvent.Notify))
 						{
 							MaxBlend = ActivePlayer.EffectiveBlendWeight;
 							MaxNotifyRecord = &ActivePlayer;
@@ -980,18 +1228,20 @@ void FIKRetargetRelativeIKOp::UpdateRelativeIKNotifyInfoAnimSeq(UAnimInstance* S
 	{
 		return;
 	}
-	
-	CacheSourceAnimSequence = Cast<UAnimSequence>(MaxNotifyRecord->SourceAsset);
 
-	// NOTE: Only grabs first bake notify
 	ResetCacheNotifyInfo();
+	CacheSourceAnimSequence = Cast<UAnimSequence>(MaxNotifyRecord->SourceAsset);
+
+	// TODO: This loop only really works w/ at most one of each notify
 	for (const FAnimNotifyEvent& NotifyEvent : CacheSourceAnimSequence->Notifies)
 	{
-		URelativeBodyBakeAnimNotify* BakedNotifyInfo = Cast<URelativeBodyBakeAnimNotify>(NotifyEvent.Notify);
-		if (BakedNotifyInfo)
+		if (URelativeBodyBakeAnimNotify* BakedNotifyInfo = Cast<URelativeBodyBakeAnimNotify>(NotifyEvent.Notify))
 		{
 			UpdateCacheNotifyInfo(BakedNotifyInfo);
-			return;
+		}
+		else if (URelativePropsBakeAnimNotify* PropInfo = Cast<URelativePropsBakeAnimNotify>(NotifyEvent.Notify))
+		{
+			UpdateCachePropNotifyInfo(PropInfo);
 		}
 	}
 }
@@ -1002,6 +1252,9 @@ void FIKRetargetRelativeIKOp::ResetCacheNotifyInfo()
 	SegmentStartTime = -1.0f;
 	SegmentEndTime = -1.0f;
 	CachedNotifyInfo = nullptr;
+	CachedPropNotifyInfo = nullptr;
+	CacheMeshPairOffset = 0;
+	CachePropPairOffset = 0;
 
 	// TODO: Need to bake info in accessible bipartite graph structure and cache less data!
 	// Clear all cached data on update
@@ -1009,31 +1262,38 @@ void FIKRetargetRelativeIKOp::ResetCacheNotifyInfo()
 	CacheSourceBodyEffectVertIdx.Reset();
 	CacheSourceEffectBones.Reset();
 	CacheSourceDomainBones.Reset();
+	CacheBodyBones.Reset();
 	FrameBoneTargets.Reset();
 	FrameBoneWeights.Reset();
+	CachePinSourcePropBones.Reset();
 }
 
 void FIKRetargetRelativeIKOp::UpdateCacheNotifyInfo(URelativeBodyBakeAnimNotify* NotifyInfo)
 {
 	CachedNotifyInfo = NotifyInfo;
+	CacheMeshPairOffset = FrameBoneTargets.Num() / 2;
 
 	int32 NumBodies = CachedNotifyInfo->BodyPairs.Num();
-	FrameBoneTargets.SetNumUninitialized(NumBodies);
-	FrameBoneWeights.SetNumUninitialized(NumBodies);
-
-	SampleRate = FMath::RoundToFloat(1.0f / (CachedNotifyInfo->BodyPairsSampleTime[1] - CachedNotifyInfo->BodyPairsSampleTime[0]));
+	FrameBoneTargets.AddDefaulted(NumBodies);
+	FrameBoneWeights.AddZeroed(NumBodies);
 
-	TMap<FName,int32> CheckSourceBonesMap;
+	MeshSampleRate = 30.0f;
+	if (CachedNotifyInfo->BodyPairsSampleTime.Num() >= 2)
+	{
+		const float DeltaT = CachedNotifyInfo->BodyPairsSampleTime[1] - CachedNotifyInfo->BodyPairsSampleTime[0];
+		MeshSampleRate = FMath::RoundToFloat(1.0f / DeltaT);
+	}
+	
 	for (int BodyIdx = 0; BodyIdx < NumBodies; ++BodyIdx)
 	{
 		FName SourceBodyBone = CachedNotifyInfo->BodyPairs[BodyIdx];
-		if (CheckSourceBonesMap.Contains(SourceBodyBone))
+		if (CacheSourceEffectBones.Contains(SourceBodyBone))
 		{
 			continue;
 		}
-
-		CheckSourceBonesMap.Emplace(SourceBodyBone, CacheSourceEffectBones.Num());
+		
 		CacheSourceEffectBones.Add(SourceBodyBone);
+		CacheBodyBones.Add(SourceBodyBone);
 		CacheSourceBodyEffectVertIdx.AddDefaulted();
 
 		FName TargetBodyBone = ApplyBodyMap(SourceBodyBone);
@@ -1045,20 +1305,21 @@ void FIKRetargetRelativeIKOp::UpdateCacheNotifyInfo(URelativeBodyBakeAnimNotify*
 	// Cache bone-bone effectors for creating weighted goals
 	for (int PairIdx = 0; PairIdx < CachedNotifyInfo->bBodyPairsIsParentDominates.Num(); ++PairIdx)
 	{
+		int32 OutPairIdx = PairIdx + CacheMeshPairOffset;
 		bool bPrimaryFixed = CachedNotifyInfo->bBodyPairsIsParentDominates[PairIdx];
 		
 		FName PrimaryBone = CachedNotifyInfo->BodyPairs[2*PairIdx];
 		FName OtherBone = CachedNotifyInfo->BodyPairs[2*PairIdx + 1];
 
-		int32 PrimaryIdx = CheckSourceBonesMap[PrimaryBone];
-		int32 OtherIdx = CheckSourceBonesMap[OtherBone];
+		int32 PrimaryIdx = CacheSourceEffectBones.Find(PrimaryBone);
+		int32 OtherIdx = CacheSourceEffectBones.Find(OtherBone);
 
 		// Secondary (contact) Bone always affected
-		CacheSourceBodyEffectVertIdx[OtherIdx].Add(2*PairIdx+1);
+		CacheSourceBodyEffectVertIdx[OtherIdx].Add(2*OutPairIdx+1);
 		if (!bPrimaryFixed)
 		{
 			// Primary bone can be affect if e.g. hands clapping
-			CacheSourceBodyEffectVertIdx[PrimaryIdx].Add(2*PairIdx);
+			CacheSourceBodyEffectVertIdx[PrimaryIdx].Add(2*OutPairIdx);
 		}
 	}
 
@@ -1071,6 +1332,62 @@ void FIKRetargetRelativeIKOp::UpdateCacheNotifyInfo(URelativeBodyBakeAnimNotify*
 	}
 }
 
+void FIKRetargetRelativeIKOp::UpdateCachePropNotifyInfo(URelativePropsBakeAnimNotify* NotifyInfo)
+{
+	CachedPropNotifyInfo = NotifyInfo;
+	CachePropPairOffset = FrameBoneTargets.Num() / 2;
+
+	int32 NumBones = CachedPropNotifyInfo->PropsPairs.Num();
+	FrameBoneTargets.AddDefaulted(NumBones);
+	FrameBoneWeights.AddZeroed(NumBones);
+
+	PropSampleRate = 30.0f;
+	if (CachedPropNotifyInfo->PropsPairsSampleTime.Num() >= 2)
+	{
+		const float DeltaT = CachedPropNotifyInfo->PropsPairsSampleTime[1] - CachedPropNotifyInfo->PropsPairsSampleTime[0];
+		PropSampleRate = FMath::RoundToFloat(1.0f / DeltaT);
+	}
+	
+	for (int BoneIdx = 0; BoneIdx < NumBones; ++BoneIdx)
+	{
+		FName SourceBodyBone = CachedPropNotifyInfo->PropsPairs[BoneIdx];
+		if (CacheSourceEffectBones.Contains(SourceBodyBone))
+		{
+			continue;
+		}
+		
+		CacheSourceEffectBones.Add(SourceBodyBone);
+		CacheSourceBodyEffectVertIdx.AddDefaulted();
+
+		FName TargetBodyBone = ApplyBodyMap(SourceBodyBone);
+		UpdateSourceBoneMap(SourceBodyBone);
+		UpdateTargetBoneMap(TargetBodyBone);
+		UpdateBoneMapTfm(SourceBodyBone, TargetBodyBone);
+	}
+
+	// Cache bone-bone effectors for creating weighted goals
+	const int32 PairCount = CachedPropNotifyInfo->PropsPairs.Num() / 2;
+	for (int PairIdx = 0; PairIdx < PairCount; ++PairIdx)
+	{
+		int32 OutPairIdx = PairIdx + CachePropPairOffset;
+		bool bPrimaryFixed = false;
+		
+		FName PropBone = CachedPropNotifyInfo->PropsPairs[2*PairIdx];
+		FName BodyBone = CachedPropNotifyInfo->PropsPairs[2*PairIdx + 1];
+
+		int32 PropIdx = CacheSourceEffectBones.Find(PropBone);
+		int32 BodyIdx = CacheSourceEffectBones.Find(BodyBone);
+
+		// TODO: Only add IK goal bones here!
+		CacheSourceBodyEffectVertIdx[BodyIdx].Add(2*OutPairIdx+1);
+		// Prop attach bones always allowed to move
+		CacheSourceBodyEffectVertIdx[PropIdx].Add(2*OutPairIdx);
+
+		// Keep track of prop attach bones for post-solve pin/update
+		CachePinSourcePropBones.AddUnique(PropBone);
+	}
+}
+
 
 void FIKRetargetRelativeIKOp::OnAddedToStack(const UIKRetargeter* InRetargetAsset, const FIKRetargetOpBase* InParentOp)
 {
@@ -1329,7 +1646,109 @@ void FIKRetargetRelativeIKOp::DebugDrawBodyCoords(FPrimitiveDrawInterface* InPDI
 	}
 }
 
+void FIKRetargetRelativeIKOp::UpdateDebugInfo(
+	double SourceScale,
+	const TArray<FTransform>& SourcePose,
+	const TArray<FTransform>& TargetPose,
+	const FDebugRelativeIKDrawInfo& LocalDebugInfo)
+{
+	if (!Settings.bDebugDraw)
+	{
+		return;
+	}
+	
+	FScopeLock ScopeLock(&DebugDataMutex);
+
+	DebugDrawInfo.TargetGoals = LocalDebugInfo.TargetGoals;
+	DebugDrawInfo.PairRetargetInfo = LocalDebugInfo.PairRetargetInfo;
+	DebugDrawInfo.SourcePairVerts = LocalDebugInfo.SourcePairVerts;
+	DebugDrawInfo.TargetPairVerts = LocalDebugInfo.TargetPairVerts;
+
+	// Cache debug info for source bodies
+	DebugDrawInfo.SourceBodyInfo.Reset(CacheSourceEffectBones.Num());
+	for (FName SourceBoneName : CacheSourceEffectBones)
+	{
+		if (!CacheSourceSkelIndices.Contains(SourceBoneName))
+		{
+			continue;
+		}
+
+		int32 SourcePoseIdx = CacheSourceSkelIndices[SourceBoneName];
+		// Need to apply source-scale to bodies since it's directly applied to translations only in source pose 
+		FTransform SourceScaleTfm{FQuat::Identity, FVector::ZeroVector, FVector(SourceScale)};
+		DebugDrawInfo.SourceBodyInfo.Emplace(SourceBoneName, SourceScale, SourceScaleTfm*SourcePose[SourcePoseIdx], FTransform::Identity);
+	}
+
+	// Cache debug info for target domain bodies
+	DebugDrawInfo.TargetBodyInfo.Reset(CacheSourceDomainBones.Num());
+	for (FName SourceBoneName : CacheSourceDomainBones)
+	{
+		FName TargetBoneName = ApplyBodyMap(SourceBoneName);
+		if (!CacheTargetSkelIndices.Contains(TargetBoneName))
+		{
+			continue;
+		}
+
+		FTransform TargetBoneDelta = GetTargetBoneDelta(TargetBoneName);
+		int32 TargetPoseIdx = CacheTargetSkelIndices[TargetBoneName];
+		DebugDrawInfo.TargetBodyInfo.Emplace(TargetBoneName, 1.0, TargetPose[TargetPoseIdx], TargetBoneDelta);
+	}
+
+	DebugDrawInfo.SourceTfmInfo.Reset(CacheSourceEffectBones.Num());
+	for (FName SourceBoneName : CacheSourceEffectBones)
+	{
+		if (!CacheSourceSkelIndices.Contains(SourceBoneName))
+		{
+			continue;
+		}
+
+		if (!FindBodyShape(SourcePhysicsAsset, SourceBoneName))
+		{
+			continue;
+		}
 
+		int32 SourcePoseIdx = CacheSourceSkelIndices[SourceBoneName];
+		
+		FBodyTransform SourceFullTfm;
+		FTransform SourceBodyRot = FTransform(GetBodyRotation(SourcePhysicsAsset, SourceBoneName));
+		ComputeSourceBodyTransform(SourceFullTfm, SourceBoneName, SourcePose[SourcePoseIdx], SourceScale);
+
+		FVector SourceBodyCenter = ApplyBodyTransform(SourceFullTfm, FVector::ZeroVector);
+		FVector SourceBodyX = ApplyBodyTransform(SourceFullTfm, SourceBodyRot.GetUnitAxis(EAxis::X));
+		FVector SourceBodyY = ApplyBodyTransform(SourceFullTfm, SourceBodyRot.GetUnitAxis(EAxis::Y));
+		FVector SourceBodyZ = ApplyBodyTransform(SourceFullTfm, SourceBodyRot.GetUnitAxis(EAxis::Z));
+		
+		DebugDrawInfo.SourceTfmInfo.Emplace(SourceBodyCenter, SourceBodyX, SourceBodyY, SourceBodyZ);
+	}
+	
+	DebugDrawInfo.TargetTfmInfo.Reset(CacheSourceEffectBones.Num());
+	for (FName SourceBoneName : CacheSourceEffectBones)
+	{
+		FName TargetBoneName = ApplyBodyMap(SourceBoneName);
+		if (!CacheTargetSkelIndices.Contains(TargetBoneName))
+		{
+			continue;
+		}
+
+		if (!FindBodyShape(SourcePhysicsAsset, SourceBoneName))
+		{
+			continue;
+		}
+		
+		int32 TargetPoseIdx = CacheTargetSkelIndices[TargetBoneName];
+
+		FBodyTransform TargetFullTfm;
+		FTransform SourceBodyRot = FTransform(GetBodyRotation(SourcePhysicsAsset, SourceBoneName));
+		ComputeTargetBodyTransform(TargetFullTfm, SourceBoneName, TargetPose[TargetPoseIdx]);
+
+		FVector TargetBodyCenter = ApplyBodyTransform(TargetFullTfm, FVector::ZeroVector);
+		FVector TargetBodyX = ApplyBodyTransform(TargetFullTfm, SourceBodyRot.GetUnitAxis(EAxis::X));
+		FVector TargetBodyY = ApplyBodyTransform(TargetFullTfm, SourceBodyRot.GetUnitAxis(EAxis::Y));
+		FVector TargetBodyZ = ApplyBodyTransform(TargetFullTfm, SourceBodyRot.GetUnitAxis(EAxis::Z));
+		
+		DebugDrawInfo.TargetTfmInfo.Emplace(TargetBodyCenter, TargetBodyX, TargetBodyY, TargetBodyZ);
+	}
+}
 
 void FIKRetargetRelativeIKOp::ResetDebugInfo()
 {
diff --git a/Engine/Plugins/Experimental/Animation/RelativeIKOp/Source/RelativeIKOp/Public/RelativeIKOp.h b/Engine/Plugins/Experimental/Animation/RelativeIKOp/Source/RelativeIKOp/Public/RelativeIKOp.h
--- a/Engine/Plugins/Experimental/Animation/RelativeIKOp/Source/RelativeIKOp/Public/RelativeIKOp.h
+++ b/Engine/Plugins/Experimental/Animation/RelativeIKOp/Source/RelativeIKOp/Public/RelativeIKOp.h
@@ -2,7 +2,6 @@
 #pragma once
 
 #include "Retargeter/IKRetargetOps.h"
-#include "Retargeter/IKRetargetSettings.h"
 #include "RelativeIKOp.generated.h"
 
 #define UE_API RELATIVEIKOP_API
@@ -18,6 +17,7 @@ class UAnimInstance;
 class UAnimSequence;
 class UPhysicsAsset;
 class URelativeBodyBakeAnimNotify;
+class URelativePropsBakeAnimNotify;
 
 // Debug info for individual body pair contributions
 struct FDebugRelativeTargetPairSpace
@@ -164,7 +164,7 @@ struct FIKRetargetRelativeIKOpSettings : public FIKRetargetOpSettingsBase
 	bool bDebugDrawBodyTransforms = false;
 
 	// Run op and display debug info but DON'T update IK Goals
-	UPROPERTY(BlueprintReadWrite, Category = "Debug")
+	UPROPERTY(EditAnywhere, BlueprintReadWrite, Category = "Debug")
 	bool bDryRun = false;
 
 	// Use distance to pair contact to blend towards direct contact
@@ -178,6 +178,14 @@ struct FIKRetargetRelativeIKOpSettings : public FIKRetargetOpSettingsBase
 	// Use target (default) when computing pair alternate representations
 	UPROPERTY(BlueprintReadWrite, Category = "Test Features")
 	bool bTestRetargetScale = true;
+
+	// Apply prop pin post-solve
+	UPROPERTY(EditAnywhere, BlueprintReadWrite, Category = "Test Features")
+	bool bApplyPostSolvePinning = true;
+	
+	// Blend from relative ik pin bone to source location
+	UPROPERTY(EditAnywhere, BlueprintReadWrite, Category = "Test Features", meta=(ClampMin = "0", ClampMax = "1"))
+	double BlendPinsToSource = 0.0;
 };
 
 USTRUCT(BlueprintType, meta = (DisplayName = "Relative IK Goals"))
@@ -198,6 +206,11 @@ struct FIKRetargetRelativeIKOp : public FIKRetargetOpBase
 		const TArray<FTransform>& InSourceGlobalPose,
 		TArray<FTransform>& OutTargetGlobalPose) override;
 
+	UE_API virtual void RunAfterParent(FIKRetargetProcessor& InProcessor,
+		const double InDeltaTime,
+		const TArray<FTransform>& InSourceGlobalPose,
+		TArray<FTransform>& OutTargetGlobalPose) override;
+
 	UE_API virtual void AnimGraphPreUpdateMainThread(USkeletalMeshComponent& SourceMeshComponent, USkeletalMeshComponent& TargetMeshComponent) override;
 
 	UE_API virtual void OnAddedToStack(const UIKRetargeter* InRetargetAsset, const FIKRetargetOpBase* InParentOp) override;
@@ -227,6 +240,46 @@ private:
 	void ResetCacheNotifyInfo();
 	void UpdateCacheNotifyInfo(URelativeBodyBakeAnimNotify* NotifyInfo);
 	void UpdateCacheChainInfo(const TArray<FTransform>& TargetPose, FName BoneName);
+	void UpdateCachePropNotifyInfo(URelativePropsBakeAnimNotify* NotifyInfo);
+
+	void ComputeFramePairSpaces(
+		double SourceScale,
+		const TArray<FTransform>& SourcePose,
+		const TArray<FTransform>& TargetPose,
+		int32 PairOffset,
+		TConstArrayView<FVector3f> FramePointView,
+		FDebugRelativeIKDrawInfo& DebugInfo,
+		TArray<FDebugRelativeTargetPairSpace>& DebugPairSpaces);
+
+	void ComputeFramePropSpaces(
+		double SourceScale,
+		const TArray<FTransform>& SourcePose,
+		const TArray<FTransform>& TargetPose,
+		int32 PairOffset,
+		TConstArrayView<FVector3f> FramePointView,
+		FDebugRelativeIKDrawInfo& DebugInfo,
+		TArray<FDebugRelativeTargetPairSpace>& DebugPairSpaces);
+
+	void UpdateTargetIKGoal(
+		TArray<FIKRigGoal>& OutIKGoals,
+		FName TargetBoneName,
+		const FTransform& TargetBonePose,
+		const TArray<int32>& EffectIndices,
+		FDebugRelativeIKDrawInfo& DebugInfo,
+		TArray<FDebugRelativeTargetPairSpace>& DebugPairSpaces);
+
+	void UpdatePropPinTransform(
+		FName SourceBoneName,
+		const FTransform& SourceBonePose,
+		const TArray<int32>& EffectIndices,
+		FDebugRelativeIKDrawInfo& DebugInfo,
+		TArray<FDebugRelativeTargetPairSpace>& DebugPairSpaces);
+
+	double ComputeTargetWeightedLocation(
+		FVector& OutTargetLocation,
+		const TArray<int32>& EffectIndices,
+		FDebugRelativeIKDrawInfo& DebugInfo,
+		TArray<FDebugRelativeTargetPairSpace>& DebugPairSpaces);
 
 	FName ApplyBodyMap(FName BodyName) const;
 	bool HasBoneDelta(FName TargetBoneName) const;
@@ -249,7 +302,7 @@ private:
 	static FVector InverseBodyTransform(const FBodyTransform& Transform, const FVector& GlobalPos);
 	
 	void ComputeFeasibilityRange(FDoubleInterval& OutFeasibleRange, FName BoneName, FVector BoneChainOffset, FVector TargetRelativeToBoneV, FVector TargetRelativeToOpposeBoneV);
-	TConstArrayView<FVector3f> ApplyTemporalSmoothing(float Time);
+	TConstArrayView<FVector3f> ApplyTemporalSmoothing(float Time, float SampleRate, int32 NumSamples, int32 NumBodies, const TArray<FVector3f>& PairLocalVerts);
 
 	FVector GoalLocBlendCompSpace(const FIKRigGoal* Goal, const FTransform& BoneTfm) const;
 
@@ -284,7 +337,8 @@ private:
 	void DebugDrawBodies(FPrimitiveDrawInterface* InPDI, const FTransform& BaseTransform, double Scale, UPhysicsAsset* PhysAsset, const TArray<FDebugDrawBodyInfo>& PhysBodies) const;
 	void DebugDrawPhysBody(FPrimitiveDrawInterface* InPDI, const FTransform& ParentTransform, double Scale, const FKShapeElem* ShapeElem, const FLinearColor& Color) const;
 	void DebugDrawBodyCoords(FPrimitiveDrawInterface* InPDI, const FTransform& BaseTransform, double Scale, UPhysicsAsset* PhysAsset, const TArray<FDebugBodyTfmInfo>& BodyTfms) const;
-	
+
+	void UpdateDebugInfo(double SourceScale, const TArray<FTransform>& SourcePose, const TArray<FTransform>& TargetPose, const FDebugRelativeIKDrawInfo& LocalDebugInfo);
 	void ResetDebugInfo();
 
 	// Collection of all debug structures
@@ -305,8 +359,12 @@ private:
 	// Cache notify-related data for faster weighted goal computations
 	TObjectPtr<const UAnimSequence> CacheSourceAnimSequence;
 	TObjectPtr<const URelativeBodyBakeAnimNotify> CachedNotifyInfo;
+	TObjectPtr<const URelativePropsBakeAnimNotify> CachedPropNotifyInfo;
+	int32 CacheMeshPairOffset = 0;
+	int32 CachePropPairOffset = 0;
 	float AnimSeqPlayHead;
-	float SampleRate;
+	float MeshSampleRate;
+	float PropSampleRate;
 
 	// Cache source/target bone-names
 	TArray<FName> SourceBoneNames;
@@ -316,6 +374,7 @@ private:
 	TArray<TArray<int32>> CacheSourceBodyEffectVertIdx;
 	TArray<FName> CacheSourceEffectBones;
 	TArray<FName> CacheSourceDomainBones;
+	TSet<FName> CacheBodyBones;
 	TArray<double> FrameBoneWeights;
 	TArray<FVector> FrameBoneTargets;
 
@@ -330,6 +389,9 @@ private:
 	// Additional per-frame ik cache maps for length lookup
 	TMap<FName,double> CacheChainLengthMap;
 	TMap<FName,FVector> CacheChainStartMap;
+	// Cache prop bones to pin/update
+	TArray<FName> CachePinSourcePropBones;
+	TArray<FTransform> CachePinTargetPropTransforms;
 
 	// Cached montage instance/segment info for early-out matching in preupdate
 	const FAnimMontageInstance* MontageInstance;
