RelativeIKOp: Fix bughawk-identified issues in relative ik prop integration code.

diff --git a/Engine/Plugins/Experimental/Animation/RelativeIKOp/Source/RelativeIKOp/Private/RelativeIKOp.cpp b/Engine/Plugins/Experimental/Animation/RelativeIKOp/Source/RelativeIKOp/Private/RelativeIKOp.cpp
--- a/Engine/Plugins/Experimental/Animation/RelativeIKOp/Source/RelativeIKOp/Private/RelativeIKOp.cpp
+++ b/Engine/Plugins/Experimental/Animation/RelativeIKOp/Source/RelativeIKOp/Private/RelativeIKOp.cpp
@@ -359,11 +359,11 @@ void FIKRetargetRelativeIKOp::ComputeFramePairSpaces(
 		FVector WeightedTargetPrimary = WeightedTarget + RelativePrimaryOffset;
 		FVector WeightedTargetOther = WeightedTarget + RelativeOtherOffset;
 		
-		FrameBoneTargets[2*PairIdx] = WeightedTargetPrimary + TargetPrimaryBoneOffset;
-		FrameBoneTargets[2*PairIdx + 1] = WeightedTargetOther + TargetOtherBoneOffset;
+		FrameBoneTargets[2*OutPairIdx] = WeightedTargetPrimary + TargetPrimaryBoneOffset;
+		FrameBoneTargets[2*OutPairIdx + 1] = WeightedTargetOther + TargetOtherBoneOffset;
 
-		FrameBoneWeights[2*PairIdx] = Weight;
-		FrameBoneWeights[2*PairIdx + 1] = Weight;
+		FrameBoneWeights[2*OutPairIdx] = Weight;
+		FrameBoneWeights[2*OutPairIdx + 1] = Weight;
 
 #if WITH_EDITOR
 		// Debug draw verts
@@ -396,6 +396,11 @@ void FIKRetargetRelativeIKOp::ComputeFramePropSpaces(
 	TArray<FDebugRelativeTargetPairSpace>& DebugPairSpaces)
 {
 	int32 PropCount = CachedPropNotifyInfo->PropsOffsetTransforms.Num();
+	if (PropCount <= 0)
+	{
+		return;
+	}
+	
 	int32 BodyCount = CachedPropNotifyInfo->PropsPairs.Num() / PropCount / 2;
 
 	FTransform SourceScaleTfm(FQuat::Identity, FVector::ZeroVector, FVector(SourceScale));
diff --git a/Engine/Plugins/Experimental/Animation/RelativeIKOp/Source/RelativeIKOp/Public/RelativeIKOp.h b/Engine/Plugins/Experimental/Animation/RelativeIKOp/Source/RelativeIKOp/Public/RelativeIKOp.h
--- a/Engine/Plugins/Experimental/Animation/RelativeIKOp/Source/RelativeIKOp/Public/RelativeIKOp.h
+++ b/Engine/Plugins/Experimental/Animation/RelativeIKOp/Source/RelativeIKOp/Public/RelativeIKOp.h
@@ -270,7 +270,7 @@ private:
 
 	void UpdatePropPinTransform(
 		FName SourceBoneName,
-		const FTransform& SourceBonePose,
+		const FTransform& TargetBonePose,
 		const TArray<int32>& EffectIndices,
 		FDebugRelativeIKDrawInfo& DebugInfo,
 		TArray<FDebugRelativeTargetPairSpace>& DebugPairSpaces);
