BodyIntersectIKOp: Don't itgnore physics asset and other settings in runtime copy.

diff --git a/Engine/Plugins/Experimental/Animation/RelativeIKOp/Source/BodyIntersectIKOp/Private/BodyIntersectIKOp.cpp b/Engine/Plugins/Experimental/Animation/RelativeIKOp/Source/BodyIntersectIKOp/Private/BodyIntersectIKOp.cpp
--- a/Engine/Plugins/Experimental/Animation/RelativeIKOp/Source/BodyIntersectIKOp/Private/BodyIntersectIKOp.cpp
+++ b/Engine/Plugins/Experimental/Animation/RelativeIKOp/Source/BodyIntersectIKOp/Private/BodyIntersectIKOp.cpp
@@ -26,7 +26,7 @@ const UClass* FIKRetargetBodyIntersectIKOpSettings::GetControllerType() const
 
 void FIKRetargetBodyIntersectIKOpSettings::CopySettingsAtRuntime(const FIKRetargetOpSettingsBase* InSettingsToCopyFrom)
 {
-	static TArray<FName> PropertiesToIgnore = {"PhysicsAssetOverride", "IntersectGoals", "IntersectBodies"};
+	static TArray<FName> PropertiesToIgnore = {};
 	FIKRetargetOpBase::CopyStructProperties(
 		FIKRetargetBodyIntersectIKOpSettings::StaticStruct(),
 		InSettingsToCopyFrom,
