BodyIntersectIK:
+ Support for Capsules and Box shape
+ Use Physics Asset's collisoin table
+ Better visual debug

diff --git a/Engine/Plugins/Experimental/Animation/RelativeIKOp/Source/BodyIntersectIKOp/Private/BodyIntersectIKOp.cpp b/Engine/Plugins/Experimental/Animation/RelativeIKOp/Source/BodyIntersectIKOp/Private/BodyIntersectIKOp.cpp
--- a/Engine/Plugins/Experimental/Animation/RelativeIKOp/Source/BodyIntersectIKOp/Private/BodyIntersectIKOp.cpp
+++ b/Engine/Plugins/Experimental/Animation/RelativeIKOp/Source/BodyIntersectIKOp/Private/BodyIntersectIKOp.cpp
@@ -3,6 +3,7 @@
 #include "BodyIntersectIKOp.h"
 
 #include "IKRigObjectVersion.h"
+#include "Engine/SkeletalMesh.h"
 #include "PhysicsEngine/PhysicsAsset.h"
 #include "PhysicsEngine/SkeletalBodySetup.h"
 #include "Retargeter/IKRetargetProcessor.h"
@@ -14,6 +15,7 @@
 #include "Engine/EngineTypes.h"
 #endif
 
+
 #include UE_INLINE_GENERATED_CPP_BY_NAME(BodyIntersectIKOp)
 
 #define LOCTEXT_NAMESPACE "BodyIntersectIKOp"
@@ -68,10 +70,11 @@ bool FIKRetargetBodyIntersectIKOp::Initialize(
 	}
 
 	UpdateCacheSkelInfo(InTargetSkeleton);
-	
+
 	const FIKRigGoalContainer& GoalContainer = InProcessor.GetIKRigGoalContainer();
-	for (const FIKGoalIntersectSphereSettings& GoalSettings : Settings.IntersectGoalSettings)
+	for (int i = 0; i < Settings.IntersectGoalSettings.Num(); i++)
 	{
+		FIKGoalIntersectShapeSettings& GoalSettings = Settings.IntersectGoalSettings[i];
 		const FIKRigGoal* IKGoal = GoalContainer.FindGoalByName(GoalSettings.Goal);
 		if (!IKGoal)
 		{
@@ -107,29 +110,31 @@ void FIKRetargetBodyIntersectIKOp::Run(
 
 #if WITH_EDITOR
 	FDebugBodyIntersectDrawInfo LocalDebugInfo;
+	LocalDebugInfo.TargetIntersectDetected.SetNumZeroed(Settings.IntersectBodies.Num());
 #endif //WITH_EDITOR
 
 	FIKRigGoalContainer& GoalContainer = InProcessor.GetIKRigGoalContainer();
-	for (const FIKGoalIntersectSphereSettings& GoalSettings : Settings.IntersectGoalSettings)
+	for (int i = 0; i < Settings.IntersectGoalSettings.Num(); i++)
 	{
+		FIKGoalIntersectShapeSettings& GoalSettings = Settings.IntersectGoalSettings[i];
 		FIKRigGoal* IKGoal = GoalContainer.FindGoalByName(GoalSettings.Goal);
 		if (!IKGoal || !IKGoal->bEnabled)
 		{
 			continue;
 		}
 
-		const FName TargetBoneName = IKGoal->BoneName;
-		if (!CacheTargetSkelIndices.Contains(TargetBoneName))
+		const FName GoalBoneName = IKGoal->BoneName;
+		if (!CacheTargetSkelIndices.Contains(GoalBoneName))
 		{
-			// UE_LOG(LogAnimation, Warning, TEXT("---- Could not find cache bone: %s"), *TargetBoneName.ToString());
+			// UE_LOG(LogAnimation, Warning, TEXT("---- Could not find cache bone: %s"), *GoalBoneName.ToString());
 			continue;
 		}
 
-		const int32 TargetBoneIdx = CacheTargetSkelIndices[TargetBoneName];
+		const int32 TargetBoneIdx = CacheTargetSkelIndices[GoalBoneName];
 		const FTransform TargetBoneTfm = OutTargetGlobalPose[TargetBoneIdx];
 
 		FKShapeElem* GoalShape = FindBodyShape(PhysicsAsset, IKGoal->BoneName);
-		double GoalRadius = CalcShapeSmallRadius(GoalShape) * GoalSettings.SphereScale;
+		double GoalRadius = CalcShapeSmallRadius(GoalShape); // * GoalSettings.SphereScale;
 		if (FMath::IsNearlyZero(GoalRadius))
 		{
 			continue;
@@ -137,11 +142,15 @@ void FIKRetargetBodyIntersectIKOp::Run(
 
 		// TODO: Handle actual rotation w/ a lever arm and second sphere
 		FVector CompGoalUpdate = GoalLocBlendCompSpace(IKGoal, TargetBoneTfm);
-		FVector CompSphereOffset = TargetBoneTfm.TransformVector(GoalSettings.SphereOffset);
+		FVector CompSphereOffset = TargetBoneTfm.TransformVector(GoalSettings.GoalShapeOffset);
 		FVector SphereOffsetGoal = CompGoalUpdate + CompSphereOffset;
-		for (FName BodyName : Settings.IntersectBodies)
+		FTransform GoalTfm = FTransform(TargetBoneTfm.Rotator(), SphereOffsetGoal, TargetBoneTfm.GetScale3D()*GoalSettings.GoalShapeScale);
+		
+		for (int j = 0; j < Settings.IntersectBodies.Num(); ++j)
 		{
-			if (TargetBoneName == BodyName)
+			FName BodyName = Settings.IntersectBodies[j];
+			
+			if (BodyName == GoalBoneName)
 			{
 				continue;
 			}
@@ -150,6 +159,11 @@ void FIKRetargetBodyIntersectIKOp::Run(
 			{
 				continue;
 			}
+
+			if (PhysicsAsset->FindBodyIndex(GoalBoneName) == INDEX_NONE || PhysicsAsset->FindBodyIndex(BodyName) == INDEX_NONE || !PhysicsAsset->IsCollisionEnabled(PhysicsAsset->FindBodyIndex(GoalBoneName), PhysicsAsset->FindBodyIndex(BodyName)))
+			{
+				continue;
+			}
 				
 			int32 BodyPoseIdx = CacheTargetSkelIndices[BodyName];
 			FTransform BodyTfm = OutTargetGlobalPose[BodyPoseIdx];
@@ -160,10 +174,14 @@ void FIKRetargetBodyIntersectIKOp::Run(
 			}
 
 			FVector DeltaDir = FVector::ZeroVector;
-			double Dist = CalcIntersectionDelta(SphereOffsetGoal, GoalRadius, BodyTfm, ShapeElem, DeltaDir);
+			double Dist = CalcIntersectionPairDelta(GoalTfm, GoalShape, BodyTfm, ShapeElem, DeltaDir);
 			if (Dist > 0.0)
 			{
 				SphereOffsetGoal += Dist * DeltaDir;
+				GoalTfm.SetTranslation(SphereOffsetGoal);
+#if WITH_EDITOR
+				LocalDebugInfo.TargetIntersectDetected[j] = true;
+#endif //WITH_EDITOR
 			}
 		}
 		
@@ -171,7 +189,7 @@ void FIKRetargetBodyIntersectIKOp::Run(
 #if WITH_EDITOR
 		if (Settings.bDebugDraw)
 		{
-			LocalDebugInfo.TestSpheres.Emplace(SphereOffsetGoal, GoalRadius);
+			LocalDebugInfo.GoalIntersectTfms.Emplace(GoalBoneName, GoalTfm);
 		}
 #endif //WITH_EDITOR
 	}
@@ -181,10 +199,11 @@ void FIKRetargetBodyIntersectIKOp::Run(
 	{
 		FScopeLock ScopeLock(&DebugDataMutex);
 		
-		DebugDrawInfo.TestSpheres = LocalDebugInfo.TestSpheres;
+		DebugDrawInfo.GoalIntersectTfms = LocalDebugInfo.GoalIntersectTfms;
 
 		// Cache debug info for target intersect bodies
 		DebugDrawInfo.TargetIntersectTfms.Reset(Settings.IntersectBodies.Num());
+		DebugDrawInfo.TargetIntersectDetected = LocalDebugInfo.TargetIntersectDetected;
 		for (FName TargetBoneName : Settings.IntersectBodies)
 		{
 			if (!CacheTargetSkelIndices.Contains(TargetBoneName))
@@ -249,9 +268,375 @@ void FIKRetargetBodyIntersectIKOp::SetGoalPosFromCompSpace(FIKRigGoal* Goal, con
 	}
 }
 
+double FIKRetargetBodyIntersectIKOp::CalcIntersectionPairDelta(const FTransform& GoalShapeTfm, FKShapeElem* GoalShapeElem, const FTransform& TargetShapeTfm, FKShapeElem* TargetShapeElem, FVector& OutDeltaDir)
+{
+	if (!GoalShapeElem || !TargetShapeElem)
+	{
+		return -1.0;
+	}
+	
+	double Dist;
+	EAggCollisionShape::Type GoalShapeType = GoalShapeElem->GetShapeType();
+	EAggCollisionShape::Type TargetShapeType = TargetShapeElem->GetShapeType();
+	bool NeedInverse;
+	switch (CalcBodyIntersectPairType(GoalShapeElem, TargetShapeElem, NeedInverse))
+	{
+	case EBodyIntersectPairType::SphereSphere:
+		{
+			FKSphereElem* SphereElem1 = static_cast<FKSphereElem*>(GoalShapeElem);
+			FKSphereElem* SphereElem2 = static_cast<FKSphereElem*>(TargetShapeElem);
+			Dist = CalcSphereSphereShapeDelta(OutDeltaDir, GoalShapeTfm, SphereElem1, TargetShapeTfm, SphereElem2);
+			break;
+		}
+
+	case EBodyIntersectPairType::SphereBox:
+		{
+			if (NeedInverse)
+			{
+				FKSphereElem* SphereElem1 = static_cast<FKSphereElem*>(TargetShapeElem);
+				FKBoxElem* BoxElem2 = static_cast<FKBoxElem*>(GoalShapeElem);
+				Dist = CalcSphereBoxShapeDelta(OutDeltaDir, TargetShapeTfm, SphereElem1, GoalShapeTfm, BoxElem2);
+				OutDeltaDir *= -1;
+			}
+			else
+			{
+				FKSphereElem* SphereElem1 = static_cast<FKSphereElem*>(GoalShapeElem);
+				FKBoxElem* BoxElem2 = static_cast<FKBoxElem*>(TargetShapeElem);
+				Dist = CalcSphereBoxShapeDelta(OutDeltaDir, GoalShapeTfm, SphereElem1, TargetShapeTfm, BoxElem2);
+			}
+			break;
+		}
+		
+	case EBodyIntersectPairType::SphereCapsule:
+		{
+			if (NeedInverse)
+			{
+				FKSphereElem* SphereElem1 = static_cast<FKSphereElem*>(TargetShapeElem);
+				FKSphylElem* CapsuleElem2 = static_cast<FKSphylElem*>(GoalShapeElem);
+				Dist = CalcSphereCapsuleShapeDelta(OutDeltaDir, TargetShapeTfm, SphereElem1, GoalShapeTfm, CapsuleElem2);
+				OutDeltaDir *= -1;
+			}
+			else
+			{
+				FKSphereElem* SphereElem1 = static_cast<FKSphereElem*>(GoalShapeElem);
+				FKSphylElem* CapsuleElem2 = static_cast<FKSphylElem*>(TargetShapeElem);
+				Dist = CalcSphereCapsuleShapeDelta(OutDeltaDir, GoalShapeTfm, SphereElem1, TargetShapeTfm, CapsuleElem2);
+			}
+			break;
+		}
+
+	/*case EBodyIntersectPairType::BoxBox:
+		{
+			FKBoxElem* BoxElem1 = static_cast<FKBoxElem*>(GoalShapeElem);
+			FKBoxElem* BoxElem2 = static_cast<FKBoxElem*>(TargetShapeElem);
+			Dist = CalcBoxBoxShapeDelta(OutDeltaDir, GoalShapeTfm, BoxElem1, TargetShapeTfm, BoxElem2);
+			break;
+		}*/
+
+	case EBodyIntersectPairType::BoxCapsule:
+		{
+			if (NeedInverse)
+			{
+				FKBoxElem* BoxElem1 = static_cast<FKBoxElem*>(TargetShapeElem);
+				FKSphylElem* CapsuleElem2 = static_cast<FKSphylElem*>(GoalShapeElem);
+				Dist = CalcBoxCapsuleShapeDelta(OutDeltaDir, TargetShapeTfm, BoxElem1, GoalShapeTfm, CapsuleElem2);
+				OutDeltaDir *= -1;
+			}
+			else
+			{
+				FKBoxElem* BoxElem1 = static_cast<FKBoxElem*>(GoalShapeElem);
+				FKSphylElem* CapsuleElem2 = static_cast<FKSphylElem*>(TargetShapeElem);
+				Dist = CalcBoxCapsuleShapeDelta(OutDeltaDir, GoalShapeTfm, BoxElem1, TargetShapeTfm, CapsuleElem2);
+			}
+			break;
+		}
+		
+	case EBodyIntersectPairType::CapsuleCapsule:
+		{
+			FKSphylElem* CapsuleElem1 = static_cast<FKSphylElem*>(GoalShapeElem);
+			FKSphylElem* CapsuleElem2 = static_cast<FKSphylElem*>(TargetShapeElem);
+			Dist = CalcCapsuleCapsuleShapeDelta(OutDeltaDir, GoalShapeTfm, CapsuleElem1, TargetShapeTfm, CapsuleElem2);
+			break;
+		}
+
+	default:
+		{
+			ensureMsgf(false, TEXT("Unexpected Intersection Pair Type: %d - %d"), static_cast<int>(GoalShapeType), static_cast<int>(TargetShapeType));
+			Dist = -1.0;
+		}
+	}
+
+	return Dist;
+}
+
+double FIKRetargetBodyIntersectIKOp::SegmentDistToAABB(
+	const FVector& A1, const FVector& A2,
+	const FVector& BoxMin, const FVector& BoxMax,
+	FVector& OutSeg, FVector& OutBox)
+{
+	// heuristic: try endpoints, the center projection, and orthogonal face hits
+	FVector BestS = A1;
+	FVector BestB(FMath::Clamp<FVector::FReal>(A1.X, BoxMin.X, BoxMax.X), FMath::Clamp<FVector::FReal>(A1.Y, BoxMin.Y, BoxMax.Y), FMath::Clamp<FVector::FReal>(A1.Z, BoxMin.Z, BoxMax.Z));
+	double bestD2 = FVector::DistSquared(BestS, BestB);
+
+	auto Consider = [&](const FVector& S, const FVector& C) {
+			const double d2 = FVector::DistSquared(S, C);
+			if (d2 < bestD2) { bestD2 = d2; BestS = S; BestB = C; }
+	};
+
+	// endpoints B
+	Consider(A2, FVector(FMath::Clamp<FVector::FReal>(A2.X, BoxMin.X, BoxMax.X), FMath::Clamp<FVector::FReal>(A2.Y, BoxMin.Y, BoxMax.Y), FMath::Clamp<FVector::FReal>(A2.Z, BoxMin.Z, BoxMax.Z)));
+
+	// segment point near box center
+	const FVector Center = 0.5f * (BoxMin + BoxMax);
+	const FVector P = FMath::ClosestPointOnSegment(Center, A1, A2);
+	Consider(P, FVector(FMath::Clamp<FVector::FReal>(P.X, BoxMin.X, BoxMax.X), FMath::Clamp<FVector::FReal>(P.Y, BoxMin.Y, BoxMax.Y), FMath::Clamp<FVector::FReal>(P.Z, BoxMin.Z, BoxMax.Z)));
+
+	// orthogonal face candidates
+	const FVector A1A2 = A2 - A1;
+	auto ConsiderFace = [&](int axis, double value) {
+			const double dir = (axis==0? A1A2.X : axis==1? A1A2.Y : A1A2.Z);
+			if (FMath::IsNearlyZero(dir))
+			{
+				return;
+			}
+			const double a0 = (axis==0? A1.X : axis==1? A1.Y : A1.Z);
+			const double t  = (value - a0) / dir;
+			if (t < 0.0 || t > 1.0)
+			{
+				return;
+			}
+			FVector S = A1 + t * A1A2;
+			// check other two axes in range
+			int u=(axis+1)%3, v=(axis+2)%3;
+			auto Get=[&](const FVector& V,int k)->double{ return k==0?V.X:k==1?V.Y:V.Z; };
+			if ( Get(S,u) >= Get(BoxMin,u) && Get(S,u) <= Get(BoxMax,u) &&
+				 Get(S,v) >= Get(BoxMin,v) && Get(S,v) <= Get(BoxMax,v) )
+			{
+				FVector C = S;
+				if (axis==0) C.X = value;
+				else if (axis==1) C.Y = value;
+				else              C.Z = value;
+				Consider(S, C);
+			}
+	};
+	ConsiderFace(0, BoxMin.X); ConsiderFace(0, BoxMax.X);
+	ConsiderFace(1, BoxMin.Y); ConsiderFace(1, BoxMax.Y);
+	ConsiderFace(2, BoxMin.Z); ConsiderFace(2, BoxMax.Z);
+
+	OutSeg = BestS;
+	OutBox = BestB;
+	
+	return FMath::Sqrt(bestD2);
+}
+
+double FIKRetargetBodyIntersectIKOp::CalcSphereSphereShapeDelta(FVector& OutDeltaDir, const FTransform& ShapeTfm1, const FKSphereElem* ShapeElem1, const FTransform& ShapeTfm2, const FKSphereElem* ShapeElem2)
+{
+	const double SphereRadius1 = ShapeTfm1.GetScale3D().GetAbsMin() * ShapeElem1->Radius;
+	const FVector WorldCenter1 = ShapeTfm1.TransformPositionNoScale(ShapeElem1->Center);
+	const double SphereRadius2 = ShapeTfm2.GetScale3D().GetAbsMin() * ShapeElem2->Radius;
+	const FVector WorldCenter2 = ShapeTfm2.TransformPositionNoScale(ShapeElem2->Center);
+
+	const FVector Delta = WorldCenter2 - WorldCenter1;
+	const double DeltaLen = Delta.Size();
+	const double DistToEdge = SphereRadius2 + SphereRadius1 - DeltaLen;
+
+	OutDeltaDir = (DeltaLen > UE_SMALL_NUMBER) ? - Delta / DeltaLen : FVector::ZeroVector;
+	return DistToEdge;
+}
+
+double FIKRetargetBodyIntersectIKOp::CalcSphereBoxShapeDelta(FVector& OutDeltaDir, const FTransform& ShapeTfm1, const FKSphereElem* ShapeElem1, const FTransform& ShapeTfm2, const FKBoxElem* ShapeElem2)
+{
+	const double SphereRadius1 = ShapeTfm1.GetScale3D().GetAbsMin() * ShapeElem1->Radius;
+	const FVector WorldCenter1 = ShapeTfm1.TransformPositionNoScale(ShapeElem1->Center);
+	
+	FTransform ShapeTfm2NoScale = ShapeTfm2;
+	ShapeTfm2NoScale.RemoveScaling();
+	const FTransform LocalToWorldTM2 = ShapeElem2->GetTransform() * ShapeTfm2NoScale;
+	const FVector SphereBoxWorldCenter2 = LocalToWorldTM2.InverseTransformPositionNoScale(WorldCenter1);
+	const FVector HalfBoxXYZ = 0.5 * FVector(ShapeElem2->X, ShapeElem2->Y, ShapeElem2->Z) * ShapeTfm2.GetScale3D();
+	
+	const FVector ClosestLocalPosition(FMath::Clamp<FVector::FReal>(SphereBoxWorldCenter2.X, -HalfBoxXYZ.X, HalfBoxXYZ.X), FMath::Clamp<FVector::FReal>(SphereBoxWorldCenter2.Y, -HalfBoxXYZ.Y, HalfBoxXYZ.Y), FMath::Clamp<double>(SphereBoxWorldCenter2.Z, -HalfBoxXYZ.Z, HalfBoxXYZ.Z));
+	double DistToEdge = UE_BIG_NUMBER;
+	if (ClosestLocalPosition == SphereBoxWorldCenter2)
+	{
+		const FVector DistBoxMax = HalfBoxXYZ - SphereBoxWorldCenter2;
+		const FVector DistBoxMin = HalfBoxXYZ + SphereBoxWorldCenter2;
+		for (int i = 0; i < 3; i++)
+		{
+			if (DistToEdge > DistBoxMax[i])
+			{
+				DistToEdge = DistBoxMax[i];
+				OutDeltaDir = FVector::ZeroVector;
+				OutDeltaDir[i] = 1;
+			}
+			if (DistToEdge > DistBoxMin[i])
+			{
+				DistToEdge = DistBoxMin[i];
+				OutDeltaDir = FVector::ZeroVector;
+				OutDeltaDir[i] = -1;
+			}
+		}
+		DistToEdge += SphereRadius1;
+	}
+	else
+	{
+		const FVector Delta = SphereBoxWorldCenter2-ClosestLocalPosition;
+		const double DeltaLen = Delta.Size();
+		DistToEdge = SphereRadius1 - DeltaLen;
+		OutDeltaDir = Delta / DeltaLen;
+	}
+	
+	OutDeltaDir = LocalToWorldTM2.TransformVectorNoScale(OutDeltaDir);
+	return DistToEdge;
+}
+
+double FIKRetargetBodyIntersectIKOp::CalcSphereCapsuleShapeDelta(FVector& OutDeltaDir, const FTransform& ShapeTfm1, const FKSphereElem* ShapeElem1, const FTransform& ShapeTfm2, const FKSphylElem* ShapeElem2)
+{
+	const double SphereRadius1 = ShapeTfm1.GetScale3D().GetAbsMin() * ShapeElem1->Radius;
+
+	FTransform ShapeTfm2NoScale = ShapeTfm2;
+	ShapeTfm2NoScale.RemoveScaling();
+	const FTransform LocalToWorldTM2 = ShapeElem2->GetTransform() * ShapeTfm2NoScale;
+	const FVector Scale3DAbs2 = ShapeTfm2.GetScale3D().GetAbs();
+	const double CapsuleRadius2 = ShapeElem2->GetScaledRadius(Scale3DAbs2);
+	const double CapsuleHalfLen2 = 0.5f * ShapeElem2->GetScaledCylinderLength(Scale3DAbs2);
+	
+	const FVector A1 = ShapeTfm1.TransformPositionNoScale(FVector(ShapeElem1->Center));
+	const FVector B1 = LocalToWorldTM2.TransformPositionNoScale(FVector(0.f, 0.f, CapsuleHalfLen2));
+	const FVector B2 = LocalToWorldTM2.TransformPositionNoScale(FVector(0.f, 0.f, -CapsuleHalfLen2));
+	const FVector P2 = FMath::ClosestPointOnSegment(A1, B1, B2);
+
+	const FVector Delta = P2 - A1;
+	const double DeltaLen = Delta.Size();
+	const double DistToEdge = CapsuleRadius2 + SphereRadius1 - DeltaLen;
+	
+	OutDeltaDir = (DeltaLen > UE_SMALL_NUMBER) ? - Delta / DeltaLen : FVector::ZeroVector;
+	return DistToEdge;
+}
+
+double FIKRetargetBodyIntersectIKOp::CalcBoxBoxShapeDelta(FVector& OutDeltaDir, const FTransform& ShapeTfm1, const FKBoxElem* ShapeElem1, const FTransform& ShapeTfm2, const FKBoxElem* ShapeElem2)
+{
+	// TODO: Warning for no implementation
+	OutDeltaDir = FVector::ZeroVector;
+	return 0;
+}
+
+double FIKRetargetBodyIntersectIKOp::CalcBoxCapsuleShapeDelta(FVector& OutDeltaDir, const FTransform& ShapeTfm1, const FKBoxElem* ShapeElem1, const FTransform& ShapeTfm2, const FKSphylElem* ShapeElem2)
+{
+	FTransform ShapeTfm1NoScale = ShapeTfm1;
+	ShapeTfm1NoScale.RemoveScaling();
+	const FTransform LocalToWorldTM1 = ShapeElem1->GetTransform() * ShapeTfm1NoScale;
+	const FVector HalfBoxXYZ1 = 0.5 * FVector(ShapeElem1->X, ShapeElem1->Y, ShapeElem1->Z) * ShapeTfm1.GetScale3D();
+	
+	FTransform ShapeTfm2NoScale = ShapeTfm2;
+	ShapeTfm2NoScale.RemoveScaling();
+	const FTransform LocalToWorldTM2 = ShapeElem2->GetTransform() * ShapeTfm2NoScale;
+	const FVector Scale3DAbs2 = ShapeTfm2.GetScale3D().GetAbs();
+	const double CapsuleRadius2 = ShapeElem2->GetScaledRadius(Scale3DAbs2);
+	const double CapsuleHalfLen2 = 0.5 * ShapeElem2->GetScaledCylinderLength(Scale3DAbs2);
+
+	const FVector A1 = LocalToWorldTM1.InverseTransformPositionNoScale(LocalToWorldTM2.TransformPositionNoScale(FVector(0.f, 0.f, CapsuleHalfLen2)));
+	const FVector A2 = LocalToWorldTM1.InverseTransformPositionNoScale(LocalToWorldTM2.TransformPositionNoScale(FVector(0.f, 0.f, -CapsuleHalfLen2)));
+
+	FVector PSeg, PBox;
+	double Dist = SegmentDistToAABB(A1, A2, -HalfBoxXYZ1, HalfBoxXYZ1, PSeg, PBox);
+	if (Dist >= CapsuleRadius2)
+	{
+		return 0.;
+	}
+	
+	if (Dist > UE_SMALL_NUMBER)
+	{
+		OutDeltaDir = LocalToWorldTM1.TransformVectorNoScale((PBox-PSeg)/Dist);
+		return CapsuleRadius2-Dist;
+	}
+
+	// Use the segment midpoint as representative axis point
+	const FVector Mid = 0.5f * (A1 + A2);
+	const FVector absM(FMath::Abs(Mid.X), FMath::Abs(Mid.Y), FMath::Abs(Mid.Z));
+	const FVector clear(HalfBoxXYZ1.X - absM.X, HalfBoxXYZ1.Y - absM.Y, HalfBoxXYZ1.Z - absM.Z);
+
+	// If the axis is inside, at least one clearance should be >= 0
+	// The required push along an axis = max(0, Radius - clearance)
+	double PushX = CapsuleRadius2 - clear.X;
+	double PushY = CapsuleRadius2 - clear.Y;
+	double PushZ = CapsuleRadius2 - clear.Z;
+
+	// pick the axis with largest deficit (most penetration), but only if positive
+	int Axis = -1;
+	double DistToEdge = 0;
+
+	if (PushX > DistToEdge)
+	{
+		DistToEdge = PushX;
+		Axis = 0;
+	}
+	if (PushY > DistToEdge)
+	{
+		DistToEdge = PushY;
+		Axis = 1;
+	}
+	if (PushZ > DistToEdge)
+	{
+		DistToEdge = PushZ;
+		Axis = 2;
+	}
+	
+	if (Axis >= 0 && DistToEdge > 0.f)
+	{
+		if (Axis == 0)
+		{
+			OutDeltaDir = FVector( (Mid.X >= 0 ? 1 : -1), 0, 0 );
+		}
+		else if (Axis == 1)
+		{
+			OutDeltaDir = FVector( 0, (Mid.Y >= 0 ? 1 : -1), 0 );
+		}
+		else
+		{
+			OutDeltaDir = FVector( 0, 0, (Mid.Z >= 0 ? 1 : -1) );
+		}
+	}
+	
+	OutDeltaDir = LocalToWorldTM1.TransformVectorNoScale(OutDeltaDir);
+	return DistToEdge;
+}
+
+double FIKRetargetBodyIntersectIKOp::CalcCapsuleCapsuleShapeDelta(FVector& OutDeltaDir, const FTransform& ShapeTfm1, const FKSphylElem* ShapeElem1, const FTransform& ShapeTfm2, const FKSphylElem* ShapeElem2)
+{
+	FTransform ShapeTfm1NoScale = ShapeTfm1;
+	ShapeTfm1NoScale.RemoveScaling();
+	const FTransform LocalToWorldTM1 = ShapeElem1->GetTransform() * ShapeTfm1NoScale;
+	const FVector Scale3DAbs1 = ShapeTfm1.GetScale3D().GetAbs();
+	const double CapsuleRadius1 = ShapeElem1->GetScaledRadius(Scale3DAbs1);
+	const double CapsuleHalfLen1 = 0.5 * ShapeElem1->GetScaledCylinderLength(Scale3DAbs1);
+
+	FTransform ShapeTfm2NoScale = ShapeTfm2;
+	ShapeTfm2NoScale.RemoveScaling();
+	const FTransform LocalToWorldTM2 = ShapeElem2->GetTransform() * ShapeTfm2NoScale;
+	const FVector Scale3DAbs2 = ShapeTfm2.GetScale3D().GetAbs();
+	const double CapsuleRadius2 = ShapeElem2->GetScaledRadius(Scale3DAbs2);
+	const double CapsuleHalfLen2 = 0.5 * ShapeElem2->GetScaledCylinderLength(Scale3DAbs2);
+	
+	const FVector A1 = LocalToWorldTM1.TransformPositionNoScale(FVector(0.f, 0.f, CapsuleHalfLen1));
+	const FVector A2 = LocalToWorldTM1.TransformPositionNoScale(FVector(0.f, 0.f, -CapsuleHalfLen1));
+	const FVector B1 = LocalToWorldTM2.TransformPositionNoScale(FVector(0.f, 0.f, CapsuleHalfLen2));
+	const FVector B2 = LocalToWorldTM2.TransformPositionNoScale(FVector(0.f, 0.f, -CapsuleHalfLen2));
+	FVector P1, P2;
+	FMath::SegmentDistToSegmentSafe(A1, A2, B1, B2, P1, P2);
+
+	const FVector Delta = P2 - P1;
+	const double DeltaLen = Delta.Size();
+	const double DistToEdge = CapsuleRadius1 + CapsuleRadius2 - DeltaLen;
+	
+	OutDeltaDir = (DeltaLen > UE_SMALL_NUMBER) ? - Delta / DeltaLen : FVector::ZeroVector;
+	return DistToEdge;
+}
+
 double FIKRetargetBodyIntersectIKOp::CalcIntersectionDelta(const FVector& TargetPoint, double TargetRadius, const FTransform& ShapeTfm, FKShapeElem* ShapeElem, FVector& OutDeltaDir)
 {
-	if (!ShapeElem)
+	if (!ensure(ShapeElem))
 	{
 		return -1.0;
 	}
@@ -322,7 +707,33 @@ float FIKRetargetBodyIntersectIKOp::CalcPointShapeDist(FVector& OutDeltaDir, con
 	return DistToEdge;
 }
 
+void FIKRetargetBodyIntersectIKOp::UpdateShapeScaleForShapeType(const FKShapeElem* ShapeElem, FVector& ShapeScale)
+{
+	switch (ShapeElem->GetShapeType())
+	{
+	case EAggCollisionShape::Sphere:
+		{
+			const double ScaleX = (ShapeScale.X + ShapeScale.Y + ShapeScale.Z) / 3.;
+			ShapeScale.X = ScaleX;
+			ShapeScale.Y = ScaleX;
+			ShapeScale.Z = ScaleX;
+			break;
+		}
+
+	case EAggCollisionShape::Sphyl:
+		{
+			const double ScaleX = 0.5*(ShapeScale.X + ShapeScale.Y);
+			ShapeScale.X = ScaleX;
+			ShapeScale.Y = ScaleX;
+			break;
+		}
 
+	default:
+		{
+			return;
+		}
+	}
+}
 
 double FIKRetargetBodyIntersectIKOp::CalcShapeSmallRadius(FKShapeElem* ShapeElem)
 {
@@ -387,6 +798,25 @@ double FIKRetargetBodyIntersectIKOp::CalcShapeAvgRadius(FKShapeElem* ShapeElem)
 	}
 }
 
+EBodyIntersectPairType FIKRetargetBodyIntersectIKOp::CalcBodyIntersectPairType(const FKShapeElem* ShapeElem1, const FKShapeElem* ShapeElem2, bool &NeedInverse)
+{
+	uint8 TypeIndex1 = static_cast<uint8>(ShapeElem1->GetShapeType());
+	uint8 TypeIndex2 = static_cast<uint8>(ShapeElem2->GetShapeType());
+	NeedInverse = TypeIndex1 > TypeIndex2;
+	
+	if (NeedInverse)
+	{
+		TypeIndex1 = static_cast<uint8>(ShapeElem2->GetShapeType());
+		TypeIndex2 = static_cast<uint8>(ShapeElem1->GetShapeType());
+	}
+	uint8 PairTypeIndex =  TypeIndex2 * (TypeIndex2 + 1) / 2 + TypeIndex1;
+	if (PairTypeIndex < static_cast<uint8>(EBodyIntersectPairType::Unknown))
+	{
+		return static_cast<EBodyIntersectPairType>(PairTypeIndex);
+	}
+	return EBodyIntersectPairType::Unknown;
+}
+
 FKShapeElem* FIKRetargetBodyIntersectIKOp::FindBodyShape(const UPhysicsAsset* PhysAsset, FName BoneName)
 {
 	int32 BodyIdx = PhysAsset->FindBodyIndex(BoneName);
@@ -462,27 +892,46 @@ void FIKRetargetBodyIntersectIKOp::DebugDraw(
 		return;
 	}
 
-	for (const auto& IntersectBody : DebugDrawInfo.TargetIntersectTfms)
+	for (int i = 0; i<DebugDrawInfo.TargetIntersectTfms.Num(); ++i)
 	{
+		const TTuple<FName, FTransform>& IntersectBody = DebugDrawInfo.TargetIntersectTfms[i];
 		const FName BodyName = IntersectBody.Get<0>();
 		const FKShapeElem* ShapeElem = FindBodyShape(Settings.TargetPhysicsAssetOverride, BodyName);
 		if (!ShapeElem)
 		{
 			continue;
 		}
-		
-		FTransform CompTfm = IntersectBody.Get<1>() * InComponentTransform;
-		DebugDrawPhysBody(InPDI, CompTfm, InComponentScale, ShapeElem, FLinearColor::Yellow);
+
+		FTransform TargetShapeTfm = IntersectBody.Get<1>();
+		FVector ShapeScale = TargetShapeTfm.GetScale3D()*InComponentScale;
+		TargetShapeTfm.RemoveScaling();
+		FTransform CompTfm = TargetShapeTfm * InComponentTransform;
+		FLinearColor TargetColor = FLinearColor::Green;
+		if (DebugDrawInfo.TargetIntersectDetected[i])
+		{
+			TargetColor = FLinearColor::Yellow;
+		}
+		DebugDrawPhysBody(InPDI, CompTfm, ShapeScale, ShapeElem, TargetColor);
 	}
 
-	for (const auto& Sphere : DebugDrawInfo.TestSpheres)
+	for (const TTuple<FName, FTransform>& GoalIntersectShape : DebugDrawInfo.GoalIntersectTfms)
 	{
-		FVector Center = InComponentTransform.TransformPosition(Sphere.Get<0>());
-		DrawWireSphere(InPDI, Center, FLinearColor::Red, Sphere.Get<1>()*InComponentScale, 25, SDPG_Foreground);
+		const FName BodyName = GoalIntersectShape.Get<0>();
+		const FKShapeElem* ShapeElem = FindBodyShape(Settings.TargetPhysicsAssetOverride, BodyName);
+		if (!ShapeElem)
+		{
+			continue;
+		}
+		
+		FTransform GoalShapeTfm = GoalIntersectShape.Get<1>();
+		FVector ShapeScale = GoalShapeTfm.GetScale3D()*InComponentScale;
+		GoalShapeTfm.RemoveScaling();
+		FTransform CompTfm = GoalShapeTfm * InComponentTransform;
+		DebugDrawPhysBody(InPDI, CompTfm, ShapeScale, ShapeElem, FLinearColor::Red);
 	}
 }
 
-void FIKRetargetBodyIntersectIKOp::DebugDrawPhysBody(FPrimitiveDrawInterface* InPDI, const FTransform& ParentTransform, double Scale, const FKShapeElem* ShapeElem, const FLinearColor& Color) const
+void FIKRetargetBodyIntersectIKOp::DebugDrawPhysBody(FPrimitiveDrawInterface* InPDI, const FTransform& ParentTransform, const FVector& Scale, const FKShapeElem* ShapeElem, const FLinearColor& Color) const
 {
 	FTransform BodyFrame = ShapeElem->GetTransform() * ParentTransform;
 
@@ -503,8 +952,9 @@ void FIKRetargetBodyIntersectIKOp::DebugDrawPhysBody(FPrimitiveDrawInterface* In
 	case EAggCollisionShape::Sphyl:
 		{
 			const FKSphylElem* CapsuleElem = static_cast<const FKSphylElem*>(ShapeElem);
-			const double Radius = CapsuleElem->Radius * Scale;
-			const double HalfHeight = (0.5*CapsuleElem->Length + CapsuleElem->Radius) * Scale;
+			const FVector Scale3DAbs = Scale.GetAbs();
+			const double Radius = CapsuleElem->GetScaledRadius(Scale3DAbs);
+			const double HalfHeight = CapsuleElem->GetScaledHalfLength(Scale3DAbs);
 			DrawWireCapsule(InPDI, Translation, UnitXAxis, UnitYAxis, UnitZAxis, Color,
 						Radius, HalfHeight,25, SDPG_Foreground, 0, 1.0);
 			return;
@@ -512,7 +962,7 @@ void FIKRetargetBodyIntersectIKOp::DebugDrawPhysBody(FPrimitiveDrawInterface* In
 	case EAggCollisionShape::Sphere:
 		{
 			const FKSphereElem* SphereElem = static_cast<const FKSphereElem*>(ShapeElem);
-			const double Radius = SphereElem->Radius * Scale;
+			const double Radius = SphereElem->Radius * Scale.GetAbsMin();
 			DrawWireSphere(InPDI, Translation, Color, Radius, 25, SDPG_Foreground);
 			return;
 		}
@@ -521,14 +971,13 @@ void FIKRetargetBodyIntersectIKOp::DebugDrawPhysBody(FPrimitiveDrawInterface* In
 	}
 }
 
-
-
 void FIKRetargetBodyIntersectIKOp::ResetDebugInfo()
 {
 	FScopeLock ScopeLock(&DebugDataMutex);
 	
 	DebugDrawInfo.TargetIntersectTfms.Reset();
-	DebugDrawInfo.TestSpheres.Reset();
+	DebugDrawInfo.TargetIntersectDetected.Reset();
+	DebugDrawInfo.GoalIntersectTfms.Reset();
 }
 
 #endif //WITH_EDITOR
diff --git a/Engine/Plugins/Experimental/Animation/RelativeIKOp/Source/BodyIntersectIKOp/Public/BodyIntersectIKOp.h b/Engine/Plugins/Experimental/Animation/RelativeIKOp/Source/BodyIntersectIKOp/Public/BodyIntersectIKOp.h
--- a/Engine/Plugins/Experimental/Animation/RelativeIKOp/Source/BodyIntersectIKOp/Public/BodyIntersectIKOp.h
+++ b/Engine/Plugins/Experimental/Animation/RelativeIKOp/Source/BodyIntersectIKOp/Public/BodyIntersectIKOp.h
@@ -13,31 +13,46 @@
 struct FIKRigGoal;
 struct FKShapeElem;
 struct FKSphereElem;
+struct FKBoxElem;
 struct FKSphylElem;
 class UPhysicsAsset;
+class USkeletalMesh;
+
+enum class EBodyIntersectPairType : uint8 
+{
+	SphereSphere,
+	SphereBox,
+	BoxBox,
+	SphereCapsule,
+	BoxCapsule,
+	CapsuleCapsule,
+	
+	Unknown
+};
 
 // Used to hold all debug draw info for convenience
 struct FDebugBodyIntersectDrawInfo
 {
 	// Intersect info
 	TArray<TTuple<FName,FTransform>> TargetIntersectTfms;
-	TArray<TTuple<FVector,double>> TestSpheres;
+	TArray<bool> TargetIntersectDetected;
+	TArray<TTuple<FName,FTransform>> GoalIntersectTfms;
 };
 
 USTRUCT(BlueprintType, meta = (DisplayName = "Goal Intersect Settings"))
-struct FIKGoalIntersectSphereSettings
+struct FIKGoalIntersectShapeSettings
 {
 	GENERATED_BODY()
 
 	// The IK Goal Name (See Run IK Rig Op) to intersect
 	UPROPERTY(EditAnywhere, BlueprintReadWrite, Category = "Goal Intersect Settings")
 	FName Goal = NAME_None;
-	// Sphere radius scale
+	// Shape scale , Todo: need to be updated with shape property
 	UPROPERTY(EditAnywhere, BlueprintReadWrite, Category = "Goal Intersect Settings", meta = (ClampMin="0.0"))
-	double SphereScale = 1.0;
-	// Local offset of sphere from goal bone
+	FVector GoalShapeScale = FVector::OneVector;
+	// Local offset of shape from goal bone
 	UPROPERTY(EditAnywhere, BlueprintReadWrite, Category = "Goal Intersect Settings")
-	FVector SphereOffset = FVector::ZeroVector;
+	FVector GoalShapeOffset = FVector::ZeroVector;
 };
 
 USTRUCT(BlueprintType, meta = (DisplayName = "Body Intersect IK Settings"))
@@ -53,12 +68,12 @@ struct FIKRetargetBodyIntersectIKOpSettings : public FIKRetargetOpSettingsBase
 	UPROPERTY(EditAnywhere, BlueprintReadWrite, Category = "Physics", meta=(ReinitializeOnEdit))
 	TObjectPtr<UPhysicsAsset> TargetPhysicsAssetOverride;
 
-	// IK Goals to sphere intersect against physics asset bodies 
+	// IK Goals Shape intersect against physics asset bodies 
 	UPROPERTY(EditAnywhere, BlueprintReadWrite, Category = "Intersection", meta=(ReinitializeOnEdit))
-	TArray<FIKGoalIntersectSphereSettings> IntersectGoalSettings;
+	TArray<FIKGoalIntersectShapeSettings> IntersectGoalSettings;
 
 	// Bone names with attached physics bodies to do trivial intersection against
-	// **NOTE: Only Sphere and Capsule bodies currently supported**
+	// **NOTE: Only Sphere, Capsule and Box bodies currently supported**
 	UPROPERTY(EditAnywhere, BlueprintReadWrite, Category = "Intersection", meta=(ReinitializeOnEdit))
 	TArray<FName> IntersectBodies;
 };
@@ -100,6 +115,18 @@ private:
 
 	FVector GoalLocBlendCompSpace(const FIKRigGoal* Goal, const FTransform& BoneTfm) const;
 	void SetGoalPosFromCompSpace(FIKRigGoal* Goal, const FTransform& BoneTfm, const FVector& CompSpaceLoc) const;
+
+	static void UpdateShapeScaleForShapeType(const FKShapeElem* ShapeElem, FVector& ShapeScale);
+
+	static double CalcIntersectionPairDelta(const FTransform& GoalShapeTfm, FKShapeElem* GoalShapeElem, const FTransform& TargetShapeTfm, FKShapeElem* TargetShapeElem, FVector& OutDeltaDir);
+	static double CalcSphereSphereShapeDelta(FVector& OutDeltaDir, const FTransform& ShapeTfm1, const FKSphereElem* ShapeElem1, const FTransform& ShapeTfm2, const FKSphereElem* ShapeElem2);
+	static double CalcSphereBoxShapeDelta(FVector& OutDeltaDir, const FTransform& ShapeTfm1, const FKSphereElem* ShapeElem1, const FTransform& ShapeTfm2, const FKBoxElem* ShapeElem2);
+	static double CalcSphereCapsuleShapeDelta(FVector& OutDeltaDir, const FTransform& ShapeTfm1, const FKSphereElem* ShapeElem1, const FTransform& ShapeTfm2, const FKSphylElem* ShapeElem2);
+	static double CalcBoxBoxShapeDelta(FVector& OutDeltaDir, const FTransform& ShapeTfm1, const FKBoxElem* ShapeElem1, const FTransform& ShapeTfm2, const FKBoxElem* ShapeElem2);
+	static double CalcBoxCapsuleShapeDelta(FVector& OutDeltaDir, const FTransform& ShapeTfm1, const FKBoxElem* ShapeElem1, const FTransform& ShapeTfm2, const FKSphylElem* ShapeElem2);
+	static double CalcCapsuleCapsuleShapeDelta(FVector& OutDeltaDir, const FTransform& ShapeTfm1, const FKSphylElem* ShapeElem1, const FTransform& ShapeTfm2, const FKSphylElem* ShapeElem2);
+
+	static double SegmentDistToAABB(const FVector& A, const FVector& B, const FVector& BoxMin, const FVector& BoxMax, FVector& OutSeg, FVector& OutBox);
 	
 	static double CalcIntersectionDelta(const FVector& TargetPoint, double TargetRadius, const FTransform& ShapeTfm, FKShapeElem* ShapeElem, FVector& OutDeltaDir);
 	static float CalcPointShapeDist(FVector& OutDeltaDir, const FVector& TargetPoint, const FTransform& ShapeTfm, const FKSphereElem* SphereElem);
@@ -110,6 +137,8 @@ private:
 
 	static FKShapeElem* FindBodyShape(const UPhysicsAsset* PhysAsset, FName BoneName);
 
+	static EBodyIntersectPairType CalcBodyIntersectPairType(const FKShapeElem* ShapeElem1, const FKShapeElem* ShapeElem2, bool &NeedInverse);
+
 #if WITH_EDITOR
 public:
 	UE_API virtual void DebugDraw(
@@ -122,7 +151,7 @@ public:
 	virtual bool HasDebugDrawing() override { return true; };
 
 private:
-	void DebugDrawPhysBody(FPrimitiveDrawInterface* InPDI, const FTransform& ParentTransform, double Scale, const FKShapeElem* ShapeElem, const FLinearColor& Color) const;
+	void DebugDrawPhysBody(FPrimitiveDrawInterface* InPDI, const FTransform& ParentTransform, const FVector& Scale, const FKShapeElem* ShapeElem, const FLinearColor& Color) const;
 	
 	void ResetDebugInfo();
 
