RelativeIK:
Udate modifier bake data

diff --git a/Engine/Plugins/Experimental/Animation/RelativeIKOp/Source/RelativeBodyAnimInfo/Public/RelativeBodyAnimNotifies.h b/Engine/Plugins/Experimental/Animation/RelativeIKOp/Source/RelativeBodyAnimInfo/Public/RelativeBodyAnimNotifies.h
--- a/Engine/Plugins/Experimental/Animation/RelativeIKOp/Source/RelativeBodyAnimInfo/Public/RelativeBodyAnimNotifies.h
+++ b/Engine/Plugins/Experimental/Animation/RelativeIKOp/Source/RelativeBodyAnimInfo/Public/RelativeBodyAnimNotifies.h
@@ -41,17 +41,20 @@ public:
 	 * @param BodyPairsIsParentDominatesIn Array indicating if parent body should be allowed to move (one entery per-body-pair) 
 	 */
 	UFUNCTION(BlueprintImplementableEvent, Category = "Animation")
-	UE_API void OnRelativeBodyDenseNotify(USkeletalMeshComponent* MeshComp, UAnimSequenceBase* Animation, int NumSamplesIn, const TArray<FName> &BodyPairsIn, const TArray<float> &BodyPairsSampleTimeIn, const TArray<FVector3f> &BodyPairsLocalReferenceIn, const TArray<bool> &BodyPairsIsParentDominatesIn) const;
+	UE_API void OnRelativeBodyDenseNotify(USkeletalMeshComponent* MeshComp, UAnimSequenceBase* Animation, int NumSamplesIn, const TMap<FName,FTransform> &OffsetTransformsForBonesIn, const TArray<FName> &BodyPairsIn, const TArray<float> &BodyPairsSampleTimeIn, const TArray<FVector3f> &BodyPairsLocalReferenceIn, const TArray<bool> &BodyPairsIsParentDominatesIn) const;
 
 	virtual void Notify(USkeletalMeshComponent* MeshComp, UAnimSequenceBase* Animation, const FAnimNotifyEventReference& EventReference) override
 	{
 		// UE_LOG(LogTemp, Log, TEXT("URelativeIKDenseAnimNotify Triggered: %d constraints"), NumSamples);
-		OnRelativeBodyDenseNotify(MeshComp, Animation, NumSamples, BodyPairs, BodyPairsSampleTime, BodyPairsLocalReference, bBodyPairsIsParentDominates);
+		OnRelativeBodyDenseNotify(MeshComp, Animation, NumSamples, OffsetTransformsForBones, BodyPairs, BodyPairsSampleTime, BodyPairsLocalReference, bBodyPairsIsParentDominates);
 	}
 	
 	UPROPERTY(VisibleAnywhere, AdvancedDisplay, meta = (Category = "Detail"))
 	int NumSamples = 0;
 	
+	UPROPERTY(VisibleAnywhere, AdvancedDisplay, meta = (Category = "Detail"))
+	TMap<FName,FTransform> OffsetTransformsForBones;
+	
 	UPROPERTY(VisibleAnywhere, AdvancedDisplay, meta = (Category = "Detail"))
 	TArray<FName> BodyPairs;
 
@@ -158,26 +161,26 @@ public:
 	 * @param BodyPairsIsParentDominatesIn Array indicating if parent body should be allowed to move (one entery per-body-pair)
 	 */
 	UFUNCTION(BlueprintImplementableEvent, Category = "Animation")
-	UE_API void OnRelativePropsDenseNotify(USkeletalMeshComponent* MeshComp, UAnimSequenceBase* Animation, int NumSamplesIn, const TArray<float>& PropsPairsSampleTimeIn, const TArray<FTransform>& PropsOffsetTransformsIn, const TArray<FName>& PropsPairsIn, const TArray<FVector3f>& PropsPairsLocalReferenceIn) const;
+	UE_API void OnRelativePropsDenseNotify(USkeletalMeshComponent* MeshComp, UAnimSequenceBase* Animation, int NumSamplesIn, const TMap<FName,FTransform> &OffsetTransformsForBonesIn, const TArray<FName>& PropsPairsIn, const TArray<float>& PropsPairsSampleTimeIn, const TArray<FVector3f>& PropsPairsLocalReferenceIn) const;
 
 	virtual void Notify(USkeletalMeshComponent* MeshComp, UAnimSequenceBase* Animation, const FAnimNotifyEventReference& EventReference) override
 	{
 		// UE_LOG(LogTemp, Log, TEXT("URelativeIKDenseAnimNotify Triggered: %d constraints"), NumSamples);
-		OnRelativePropsDenseNotify(MeshComp, Animation, NumSamples, PropsPairsSampleTime, PropsOffsetTransforms, PropsPairs, PropsPairsLocalReference);
+		OnRelativePropsDenseNotify(MeshComp, Animation, NumSamples, OffsetTransformsForBones, PropsPairs, PropsPairsSampleTime, PropsPairsLocalReference);
 	}
 
 	UPROPERTY(VisibleAnywhere, AdvancedDisplay, meta = (Category = "Detail"))
 	int NumSamples = 0;
 
 	UPROPERTY(VisibleAnywhere, AdvancedDisplay, meta = (Category = "Detail"))
-	TArray<float> PropsPairsSampleTime;
-
-	UPROPERTY(VisibleAnywhere, AdvancedDisplay, meta = (Category = "Detail"))
-	TArray<FTransform> PropsOffsetTransforms;
+	TMap<FName,FTransform> OffsetTransformsForBones;
 
 	UPROPERTY(VisibleAnywhere, AdvancedDisplay, meta = (Category = "Detail"))
 	TArray<FName> PropsPairs;
-
+	
+	UPROPERTY(VisibleAnywhere, AdvancedDisplay, meta = (Category = "Detail"))
+	TArray<float> PropsPairsSampleTime;
+	
 	UPROPERTY(VisibleAnywhere, AdvancedDisplay, meta = (Category = "Detail"))
 	TArray<FVector3f> PropsPairsLocalReference;
 };
diff --git a/Engine/Plugins/Experimental/Animation/RelativeIKOp/Source/RelativeBodyAnimUtils/Private/RelativeBodyAnimModifier.cpp b/Engine/Plugins/Experimental/Animation/RelativeIKOp/Source/RelativeBodyAnimUtils/Private/RelativeBodyAnimModifier.cpp
--- a/Engine/Plugins/Experimental/Animation/RelativeIKOp/Source/RelativeBodyAnimUtils/Private/RelativeBodyAnimModifier.cpp
+++ b/Engine/Plugins/Experimental/Animation/RelativeIKOp/Source/RelativeBodyAnimUtils/Private/RelativeBodyAnimModifier.cpp
@@ -117,6 +117,12 @@ void URelativeBodyAnimModifier::OnApply_Implementation(UAnimSequence* InAnimatio
 	UPhysicsAsset* PhysicsAsset = GetPhysicsAsset();
 	FReferenceSkeleton& RefSkeleton = SkeletalMeshAsset->GetRefSkeleton();
 
+	if (!PhysicsAsset)
+	{
+		UE_LOG(LogAnimation, Error, TEXT("Cannot create RelativeBodyAnimNotify. Reason: Must specify valid physics asset"));
+		return;
+	}
+
 	// TODO: Move cached data out of modifier (or check if it's outdated and only rebuild in that case)
 	CacheBodyDataForSourceMesh(CachedBodySourceData);
 	
@@ -172,6 +178,19 @@ void URelativeBodyAnimModifier::OnApply_Implementation(UAnimSequence* InAnimatio
 				{
 					continue;
 				}
+				
+				FName SourceBodyBoneName = RefSkeleton.GetBoneName(BodyIndicesToBoneIndices[BodyIndex1]);
+				int32 SourceBodyIdx = PhysicsAsset->FindBodyIndex(SourceBodyBoneName);
+				if (SourceBodyIdx == INDEX_NONE)
+				{
+					continue;
+				}
+				FKShapeElem* ShapeElem1 = PhysicsAsset->SkeletalBodySetups[SourceBodyIdx]->AggGeom.GetElement(0);
+				FTransform3f BodyTransform1(ShapeElem1->GetTransform());
+				FVector3f ElementScale3D1 = FRelativeBodyPhysicsUtils::CalcReferenceShapeScale3D(ShapeElem1);
+				BodyTransform1.SetScale3D(ElementScale3D1);
+				BakedAnimNotifyInfo->OffsetTransformsForBones.Emplace(SourceBodyBoneName, BodyTransform1);
+				
 				for (int32 BodyIndex2 = BodyIndex1 + 1; BodyIndex2 < NumBodies; ++BodyIndex2)
 				{
 					if (BodyIndicesToIgnore[BodyIndex2]) continue;
diff --git a/Engine/Plugins/Experimental/Animation/RelativeIKOp/Source/RelativeBodyAnimUtils/Private/RelativePropsAnimModifier.cpp b/Engine/Plugins/Experimental/Animation/RelativeIKOp/Source/RelativeBodyAnimUtils/Private/RelativePropsAnimModifier.cpp
--- a/Engine/Plugins/Experimental/Animation/RelativeIKOp/Source/RelativeBodyAnimUtils/Private/RelativePropsAnimModifier.cpp
+++ b/Engine/Plugins/Experimental/Animation/RelativeIKOp/Source/RelativeBodyAnimUtils/Private/RelativePropsAnimModifier.cpp
@@ -119,6 +119,12 @@ void URelativePropsAnimModifier::OnApply_Implementation(UAnimSequence* InAnimati
 	UPhysicsAsset* PhysicsAsset = GetPhysicsAsset();
 	FReferenceSkeleton& RefSkeleton = SkeletalMeshAsset->GetRefSkeleton();
 
+	if (!PhysicsAsset)
+	{
+		UE_LOG(LogAnimation, Error, TEXT("Cannot create RelativePropAnimNotify. Reason: Must specify valid physics asset"));
+		return;
+	}
+
 	// TODO: Move cached data out of modifier (or check if it's outdated and only rebuild in that case)
 	CacheBodyDataForSourceMesh(CachedBodySourceData);
 
@@ -169,9 +175,29 @@ void URelativePropsAnimModifier::OnApply_Implementation(UAnimSequence* InAnimati
 			{
 				return;
 			}
-
+			
+			for (int32 BodyIndex1 = 0; BodyIndex1 < NumBodies; ++BodyIndex1)
+			{
+				if (BodyIndicesToIgnore[BodyIndex1])
+				{
+					continue;
+				}
+				FName SourceBodyBoneName = RefSkeleton.GetBoneName(BodyIndicesToBoneIndices[BodyIndex1]);
+				int32 SourceBodyIdx = PhysicsAsset->FindBodyIndex(SourceBodyBoneName);
+				if (SourceBodyIdx == INDEX_NONE)
+				{
+					continue;
+				}
+				FKShapeElem* ShapeElem1 = PhysicsAsset->SkeletalBodySetups[SourceBodyIdx]->AggGeom.GetElement(0);
+				FTransform3f BodyTransform1(ShapeElem1->GetTransform());
+				FVector3f ElementScale3D1 = FRelativeBodyPhysicsUtils::CalcReferenceShapeScale3D(ShapeElem1);
+				BodyTransform1.SetScale3D(ElementScale3D1);
+				BakedAnimNotifyInfo->OffsetTransformsForBones.Emplace(SourceBodyBoneName, BodyTransform1);
+			}
+				
 			for (size_t PropsIdx = 0; PropsIdx < PropsStaticMeshs.Num(); PropsIdx++)
 			{
+				BakedAnimNotifyInfo->OffsetTransformsForBones.Emplace(PropsAttachBones[PropsIdx], PropsOffsetTransforms[PropsIdx]);
 				FMeshDescription* PropsMeshDescription = PropsStaticMeshs[PropsIdx]->GetMeshDescription(0); // LOD 0
 				if (PropsMeshDescription)
 				{
@@ -197,7 +223,6 @@ void URelativePropsAnimModifier::OnApply_Implementation(UAnimSequence* InAnimati
 			BakedAnimNotifyInfo->PropsPairsLocalReference.SetNumZeroed(BakedAnimNotifyInfo->PropsPairs.Num() * SampleNum);
 			BakedAnimNotifyInfo->PropsPairsSampleTime.SetNumZeroed(SampleNum);
 			BakedAnimNotifyInfo->NumSamples = SampleNum;
-			BakedAnimNotifyInfo->PropsOffsetTransforms = PropsOffsetTransforms;
 		}
 		
 		FScopedSlowTask BakePairVertsTask((float)SampleNum, LOCTEXT("BakeVertsTaskText", "Baking Props Pair Verts..."));
diff --git a/Engine/Plugins/Experimental/Animation/RelativeIKOp/Source/RelativeIKOp/Private/RelativeIKOp.cpp b/Engine/Plugins/Experimental/Animation/RelativeIKOp/Source/RelativeIKOp/Private/RelativeIKOp.cpp
--- a/Engine/Plugins/Experimental/Animation/RelativeIKOp/Source/RelativeIKOp/Private/RelativeIKOp.cpp
+++ b/Engine/Plugins/Experimental/Animation/RelativeIKOp/Source/RelativeIKOp/Private/RelativeIKOp.cpp
@@ -60,11 +60,10 @@ bool FIKRetargetRelativeIKOp::Initialize(
 		LOCTEXT("MissingIKRig", "{0}, is missing an IK rig. No goals can be updated. "), FText::FromName(GetName())));
 		return false;
 	}
-
-	SourcePhysicsAsset = Settings.SourcePhysicsAssetOverride;
+	
 	TargetPhysicsAsset = Settings.TargetPhysicsAssetOverride;
 
-	if (!SourcePhysicsAsset || !TargetPhysicsAsset)
+	if (!TargetPhysicsAsset)
 	{
 		InLog.LogWarning( FText::Format(
 		LOCTEXT("MissingPhysicsAssets", "{0}: Both Source and Target Physics Assets must be specified. "), FText::FromName(GetName())));
@@ -186,6 +185,7 @@ void FIKRetargetRelativeIKOp::Run(
 	}
 
 #if WITH_EDITOR
+	LocalDebugInfo.SourceDomainBones = CacheSourceDomainBones;
 	UpdateDebugInfo(SourceScale, InSourceGlobalPose, OutTargetGlobalPose, LocalDebugInfo);
 #endif
 }
@@ -273,8 +273,8 @@ void FIKRetargetRelativeIKOp::ComputeFramePairSpaces(
 		int32 SourceOtherPoseIdx = CacheSourceSkelIndices[SourceOtherBone];
 
 		// ---- TESTING ONLY - Remove these once notify baking is rotation invariant
-		FTransform SourcePrimaryBodyRot = FTransform(GetBodyRotation(SourcePhysicsAsset, SourcePrimaryBone));
-		FTransform SourceOtherBodyRot = FTransform(GetBodyRotation(SourcePhysicsAsset, SourceOtherBone));
+		FTransform SourcePrimaryBodyRot(CachedNotifyInfo->OffsetTransformsForBones[SourcePrimaryBone].GetRotation());
+		FTransform SourceOtherBodyRot(CachedNotifyInfo->OffsetTransformsForBones[SourceOtherBone].GetRotation());
 		
 		// Compute source transforms and source-space vertex distance
 		// const FVector3f& LocalPrimaryVert = CachedNotifyInfo->BodyPairsLocalReference[FramePairIdx + 2*PairIdx];
@@ -395,13 +395,11 @@ void FIKRetargetRelativeIKOp::ComputeFramePropSpaces(
 	FDebugRelativeIKDrawInfo& DebugInfo,
 	TArray<FDebugRelativeTargetPairSpace>& DebugPairSpaces)
 {
-	int32 PropCount = CachedPropNotifyInfo->PropsOffsetTransforms.Num();
-	if (PropCount <= 0)
+	int32 PairCount = CachedPropNotifyInfo->PropsPairs.Num() / 2;
+	if (PairCount <= 0)
 	{
 		return;
 	}
-	
-	int32 BodyCount = CachedPropNotifyInfo->PropsPairs.Num() / PropCount / 2;
 
 	FTransform SourceScaleTfm(FQuat::Identity, FVector::ZeroVector, FVector(SourceScale));
 	const double InvSourceScale = (Settings.bIgnoreSourceScale) ? 1.0 / SourceScale : 1.0;
@@ -412,158 +410,153 @@ void FIKRetargetRelativeIKOp::ComputeFramePropSpaces(
 	const double RetargetContactAlpha = FMath::Clamp(Settings.RetargetContactAlpha, 0.0, 1.0);
 	const double RetargetSpringAlpha = FMath::Clamp(Settings.RetargetSpringAlpha, 0.0, 1.0);
 	// Get pair target body verts and weights
-	for (int PropIdx = 0; PropIdx < PropCount; ++PropIdx)
+	// Get pair target body verts and weights
+	for (int PairIdx = 0; PairIdx < PairCount; ++PairIdx)
 	{
-		const FTransform PropTfm = CachedPropNotifyInfo->PropsOffsetTransforms[PropIdx];
-		for (int BodyIdx = 0; BodyIdx < BodyCount; ++BodyIdx)
+		const int32 OutPairIdx = PairIdx + PairOffset;
+		FrameBoneWeights[2*OutPairIdx] = 0;
+		FrameBoneWeights[2*OutPairIdx + 1] = 0;
+		FrameBoneTargets[2*OutPairIdx] = FVector::ZeroVector;
+		FrameBoneTargets[2*OutPairIdx + 1] = FVector::ZeroVector;
+
+		FName SourcePropBone = CachedPropNotifyInfo->PropsPairs[2 * PairIdx];
+		FName TargetPropBone = ApplyBodyMap(SourcePropBone);
+		FName SourceBodyBone = CachedPropNotifyInfo->PropsPairs[2*PairIdx + 1];
+		FName TargetBodyBone = ApplyBodyMap(SourceBodyBone);
+
+		if (TargetPropBone == NAME_None || TargetBodyBone == NAME_None)
 		{
-			const int32 PairIdx = BodyCount * PropIdx + BodyIdx;
-			const int32 OutPairIdx = PairIdx + PairOffset;
-			
-			FrameBoneWeights[2*OutPairIdx] = 0;
-			FrameBoneWeights[2*OutPairIdx + 1] = 0;
-			FrameBoneTargets[2*OutPairIdx] = FVector::ZeroVector;
-			FrameBoneTargets[2*OutPairIdx + 1] = FVector::ZeroVector;
-			
-			FName SourcePropBone = CachedPropNotifyInfo->PropsPairs[2*PairIdx];
-			FName SourceBodyBone = CachedPropNotifyInfo->PropsPairs[2*PairIdx + 1];
+			continue;
+		}
 
-			FName TargetPropBone = ApplyBodyMap(SourcePropBone);
-			FName TargetBodyBone = ApplyBodyMap(SourceBodyBone);
+		if (!CacheSourceSkelIndices.Contains(SourcePropBone) || !CacheSourceSkelIndices.Contains(SourceBodyBone))
+		{
+			continue;
+		}
 
-			if (TargetPropBone == NAME_None || TargetBodyBone == NAME_None)
-			{
-				continue;
-			}
+		if (!CacheTargetSkelIndices.Contains(TargetPropBone) || !CacheTargetSkelIndices.Contains(TargetBodyBone))
+		{
+			continue;
+		}
 
-			if (!CacheSourceSkelIndices.Contains(SourcePropBone) || !CacheSourceSkelIndices.Contains(SourceBodyBone))
-			{
-				continue;
-			}
+		if (!HasBoneDelta(TargetPropBone) || !HasBoneDelta(TargetBodyBone))
+		{
+			continue;
+		}
+		
+		int32 SourcePropPoseIdx = CacheSourceSkelIndices[SourcePropBone];
+		int32 SourceBodyPoseIdx = CacheSourceSkelIndices[SourceBodyBone];
 
-			if (!CacheTargetSkelIndices.Contains(TargetPropBone) || !CacheTargetSkelIndices.Contains(TargetBodyBone))
-			{
-				continue;
-			}
+		// ---- TESTING ONLY - Remove these once notify baking is rotation invariant
+		const FTransform PropOffsetTfm = CachedPropNotifyInfo->OffsetTransformsForBones[SourcePropBone];
+		FTransform SourceBodyRot(CachedPropNotifyInfo->OffsetTransformsForBones[SourceBodyBone].GetRotation());
+		FTransform TargetBodyRot = SourceBodyRot;
+		
+		// Compute source transforms and source-space vertex distance
+		// const FVector3f& LocalPrimaryVert = CachedNotifyInfo->BodyPairsLocalReference[FramePairIdx + 2*PairIdx];
+		// const FVector3f& LocalOtherVert = CachedNotifyInfo->BodyPairsLocalReference[FramePairIdx + 2*PairIdx + 1];
+		const FVector3f& LocalPropVert = FramePointView[2*PairIdx];
+		const FVector3f& LocalBodyVert = FramePointView[2*PairIdx+1];
 
-			if (!HasBoneDelta(TargetPropBone) || !HasBoneDelta(TargetBodyBone))
-			{
-				continue;
-			}
+		// TODO: Clean this up if we bake in rot invariant space
+		FVector SourceBoneBakePropV = (PropOffsetTfm).TransformPosition(FVector(LocalPropVert));
+		FVector SourceBoneBakeBodyV = (SourceBodyRot).TransformPosition(FVector(LocalBodyVert));
+		
+		FTransform SourcePropTfm = SourceScaleTfm * SourcePose[SourcePropPoseIdx];
+		FBodyTransform SourceBodyTfm;
+		ComputeSourceBodyTransform(SourceBodyTfm, SourceBodyBone, SourcePose[SourceBodyPoseIdx], SourceScale);
+		
+		FVector SourcePropCompV = SourcePropTfm.TransformPosition(SourceBoneBakePropV);
+		FVector SourceBodyCompV = ApplyBodyTransform(SourceBodyTfm, SourceBoneBakeBodyV);
 
-			int32 SourcePropPoseIdx = CacheSourceSkelIndices[SourcePropBone];
-			int32 SourceBodyPoseIdx = CacheSourceSkelIndices[SourceBodyBone];
+		// Rescale components for distance checks if "Ignore sources scale" is on
+		// NOTE: This works because we don't need absolute coordinates (just distance).
+		//       If absolute coords needed then must properly account for scale pivot
+		FVector RescaleSourcePropV = InvSourceScale * SourcePropCompV;
+		FVector RescaleSourceBodyV = InvSourceScale * SourceBodyCompV;
 
-			// ---- TESTING ONLY - Remove these once notify baking is rotation invariant
-			FTransform SourceBodyRot = FTransform(GetBodyRotation(SourcePhysicsAsset, SourceBodyBone));
-			FTransform TargetBodyRot = SourceBodyRot;
-			
-			// Compute source transforms and source-space vertex distance
-			// const FVector3f& LocalPrimaryVert = CachedNotifyInfo->BodyPairsLocalReference[FramePairIdx + 2*PairIdx];
-			// const FVector3f& LocalOtherVert = CachedNotifyInfo->BodyPairsLocalReference[FramePairIdx + 2*PairIdx + 1];
-			const FVector3f& LocalPropVert = FramePointView[2*PairIdx];
-			const FVector3f& LocalBodyVert = FramePointView[2*PairIdx+1];
-
-			// TODO: Clean this up if we bake in rot invariant space
-			FVector SourceBoneBakePropV = (PropTfm).TransformPosition(FVector(LocalPropVert));
-			FVector SourceBoneBakeBodyV = (SourceBodyRot).TransformPosition(FVector(LocalBodyVert));
-			
-			FTransform SourcePropTfm = SourceScaleTfm * SourcePose[SourcePropPoseIdx];
-			FBodyTransform SourceBodyTfm;
-			ComputeSourceBodyTransform(SourceBodyTfm, SourceBodyBone, SourcePose[SourceBodyPoseIdx], SourceScale);
-			
-			FVector SourcePropCompV = SourcePropTfm.TransformPosition(SourceBoneBakePropV);
-			FVector SourceBodyCompV = ApplyBodyTransform(SourceBodyTfm, SourceBoneBakeBodyV);
-
-			// Rescale components for distance checks if "Ignore sources scale" is on
-			// NOTE: This works because we don't need absolute coordinates (just distance).
-			//       If absolute coords needed then must properly account for scale pivot
-			FVector RescaleSourcePropV = InvSourceScale * SourcePropCompV;
-			FVector RescaleSourceBodyV = InvSourceScale * SourceBodyCompV;
-
-			double SourceVDist = FVector::Distance(RescaleSourcePropV, RescaleSourceBodyV);
-			double LinearWeight = FMath::Clamp(1.0 - (SourceVDist / Settings.DistanceThreshold), 0.0, 1.0);
-			double Weight = GetDistanceWeight(SourceVDist, Settings.DistanceThreshold, Settings.DistanceFade);
-
-			// Compute Source pair verts as represented in alternate pair body primary(_P) or other(_O) local bake space
-			FVector SourcePropV_B = InverseBodyTransform(SourceBodyTfm, SourcePropCompV);
-			FVector SourceBodyV_P = SourcePropTfm.InverseTransformPosition(SourceBodyCompV);
-			
-			//---------------------------------
-			// Get verts retargeted using physics body oriented scale
-			int32 TargetPropPoseIdx = CacheTargetSkelIndices[TargetPropBone];
-			int32 TargetBodyPoseIdx = CacheTargetSkelIndices[TargetBodyBone];
+		double SourceVDist = FVector::Distance(RescaleSourcePropV, RescaleSourceBodyV);
+		double LinearWeight = FMath::Clamp(1.0 - (SourceVDist / Settings.DistanceThreshold), 0.0, 1.0);
+		double Weight = GetDistanceWeight(SourceVDist, Settings.DistanceThreshold, Settings.DistanceFade);
+
+		// Compute Source pair verts as represented in alternate pair body primary(_P) or other(_O) local bake space
+		FVector SourcePropV_B = InverseBodyTransform(SourceBodyTfm, SourcePropCompV);
+		FVector SourceBodyV_P = SourcePropTfm.InverseTransformPosition(SourceBodyCompV);
+		
+		//---------------------------------
+		// Get verts retargeted using physics body oriented scale
+		int32 TargetPropPoseIdx = CacheTargetSkelIndices[TargetPropBone];
+		int32 TargetBodyPoseIdx = CacheTargetSkelIndices[TargetBodyBone];
 
-			FTransform TargetPropBoneDelta = GetTargetBoneDelta(TargetPropBone);
+		FTransform TargetPropBoneDelta = GetTargetBoneDelta(TargetPropBone);
 
-			FTransform TargetPropTfm = TargetPropBoneDelta * TargetPose[TargetPropPoseIdx];
-			FBodyTransform TargetBodyTfm;
-			ComputeTargetBodyTransform(TargetBodyTfm, TargetBodyBone, TargetPose[TargetBodyPoseIdx]);
+		FTransform TargetPropTfm = TargetPropBoneDelta * TargetPose[TargetPropPoseIdx];
+		FBodyTransform TargetBodyTfm;
+		ComputeTargetBodyTransform(TargetBodyTfm, TargetBodyBone, TargetPose[TargetBodyPoseIdx]);
 
-			// Compute prop vert represented in component space relative to prop (_P) and body (_B)
-			FVector TargetPropCompV_P = TargetPropTfm.TransformPosition(SourceBoneBakePropV);
-			FVector TargetPropCompV_B = ApplyBodyTransform(TargetBodyTfm, SourcePropV_B);
+		// Compute prop vert represented in component space relative to prop (_P) and body (_B)
+		FVector TargetPropCompV_P = TargetPropTfm.TransformPosition(SourceBoneBakePropV);
+		FVector TargetPropCompV_B = ApplyBodyTransform(TargetBodyTfm, SourcePropV_B);
 
-			// Compute body vert represented in component space relative to prop (_P) and body (_B)
-			FVector TargetBodyCompV_B = ApplyBodyTransform(TargetBodyTfm, SourceBoneBakeBodyV);
-			FVector TargetBodyCompV_P = TargetPropTfm.TransformPosition(SourceBodyV_P);
+		// Compute body vert represented in component space relative to prop (_P) and body (_B)
+		FVector TargetBodyCompV_B = ApplyBodyTransform(TargetBodyTfm, SourceBoneBakeBodyV);
+		FVector TargetBodyCompV_P = TargetPropTfm.TransformPosition(SourceBodyV_P);
 
-			// Get offsets from target verts (in their own body rep) to associated bone
-			FVector TargetPropBoneOffset = TargetPose[TargetPropPoseIdx].GetTranslation() - TargetPropCompV_P;
-			FVector TargetBodyBoneOffset = TargetPose[TargetBodyPoseIdx].GetTranslation() - TargetBodyCompV_B;
+		// Get offsets from target verts (in their own body rep) to associated bone
+		FVector TargetPropBoneOffset = TargetPose[TargetPropPoseIdx].GetTranslation() - TargetPropCompV_P;
+		FVector TargetBodyBoneOffset = TargetPose[TargetBodyPoseIdx].GetTranslation() - TargetBodyCompV_B;
 
-			// double SpringAlpha = (Settings.bTestDistContactAlpha) ? Weight : RetargetSpringAlpha;
-			FVector TargetRelativeToPropV = FMath::Lerp(TargetPropCompV_P,TargetBodyCompV_P, RetargetSpringAlpha);
-			FVector TargetRelativeToBodyV = FMath::Lerp(TargetPropCompV_B,TargetBodyCompV_B, RetargetSpringAlpha);
+		// double SpringAlpha = (Settings.bTestDistContactAlpha) ? Weight : RetargetSpringAlpha;
+		FVector TargetRelativeToPropV = FMath::Lerp(TargetPropCompV_P,TargetBodyCompV_P, RetargetSpringAlpha);
+		FVector TargetRelativeToBodyV = FMath::Lerp(TargetPropCompV_B,TargetBodyCompV_B, RetargetSpringAlpha);
 
-			// double TargetVDist = FVector::Distance(TargetRelativeToPrimaryV, TargetRelativeToOtherV);
-			// double TargetDRat = FMath::Min(1.0, 2.0 * FMath::Min(SourceVDist/TargetVDist, TargetVDist/SourceVDist));
-			
-			FVector RelativePropOffset = TargetPropCompV_P - TargetRelativeToPropV;
-			FVector RelativeBodyOffset = TargetBodyCompV_B - TargetRelativeToBodyV;
-			
-			// bool bDistanceWeighting = false;
-			
-			FDoubleInterval FeasibleRange(0.0, 1.0);
-			// UpdateCacheChainInfo(TargetPose, TargetBodyBone);
-			// if (bDistanceWeighting)
-			// {
-			// 	ComputeFeasibilityRange(FeasibleRange, TargetBodyBone, RelativeBodyOffset + TargetBodyBoneOffset, TargetRelativeToBodyV, TargetRelativeToPropV);
-			// }
-
-			double FeasibilityWeight = FMath::Lerp(FeasibleRange.Min, FeasibleRange.Max, LinearWeight);
-			// double ContactAlpha = (bDistanceWeighting) ? FeasibilityWeight : RetargetContactAlpha;
-			double ContactAlpha = RetargetContactAlpha;
-			FVector WeightedTarget = FMath::Lerp(TargetRelativeToBodyV, TargetRelativeToPropV, ContactAlpha);
-
-			FVector WeightedTargetProp = WeightedTarget + RelativePropOffset;
-			FVector WeightedTargetBody = WeightedTarget + RelativeBodyOffset;
-			
-			FrameBoneTargets[2*OutPairIdx] = WeightedTargetProp + TargetPropBoneOffset;
-			FrameBoneTargets[2*OutPairIdx + 1] = WeightedTargetBody + TargetBodyBoneOffset;
+		// double TargetVDist = FVector::Distance(TargetRelativeToPrimaryV, TargetRelativeToOtherV);
+		// double TargetDRat = FMath::Min(1.0, 2.0 * FMath::Min(SourceVDist/TargetVDist, TargetVDist/SourceVDist));
+		
+		FVector RelativePropOffset = TargetPropCompV_P - TargetRelativeToPropV;
+		FVector RelativeBodyOffset = TargetBodyCompV_B - TargetRelativeToBodyV;
+		
+		// bool bDistanceWeighting = false;
+		
+		FDoubleInterval FeasibleRange(0.0, 1.0);
+		// UpdateCacheChainInfo(TargetPose, TargetBodyBone);
+		// if (bDistanceWeighting)
+		// {
+		// 	ComputeFeasibilityRange(FeasibleRange, TargetBodyBone, RelativeBodyOffset + TargetBodyBoneOffset, TargetRelativeToBodyV, TargetRelativeToPropV);
+		// }
+
+		double FeasibilityWeight = FMath::Lerp(FeasibleRange.Min, FeasibleRange.Max, LinearWeight);
+		// double ContactAlpha = (bDistanceWeighting) ? FeasibilityWeight : RetargetContactAlpha;
+		double ContactAlpha = RetargetContactAlpha;
+		FVector WeightedTarget = FMath::Lerp(TargetRelativeToBodyV, TargetRelativeToPropV, ContactAlpha);
+
+		FVector WeightedTargetProp = WeightedTarget + RelativePropOffset;
+		FVector WeightedTargetBody = WeightedTarget + RelativeBodyOffset;
+		
+		FrameBoneTargets[2*OutPairIdx] = WeightedTargetProp + TargetPropBoneOffset;
+		FrameBoneTargets[2*OutPairIdx + 1] = WeightedTargetBody + TargetBodyBoneOffset;
 
-			FrameBoneWeights[2*OutPairIdx] = Weight;
-			FrameBoneWeights[2*OutPairIdx + 1] = Weight;
+		FrameBoneWeights[2*OutPairIdx] = Weight;
+		FrameBoneWeights[2*OutPairIdx + 1] = Weight;
 
 #if WITH_EDITOR
-			// Debug draw verts
-			if (Settings.bDebugDraw)
+		// Debug draw verts
+		if (Settings.bDebugDraw)
+		{
+			DebugInfo.SourcePairVerts.Add({SourcePropCompV, SourceBodyCompV, false, FVector::ZeroVector, FVector::ZeroVector});
+			if (Settings.DebugFullRetargetPairBones.Contains(TargetPropBone))
 			{
-				DebugInfo.SourcePairVerts.Add({SourcePropCompV, SourceBodyCompV, false, FVector::ZeroVector, FVector::ZeroVector});
-				if (Settings.DebugFullRetargetPairBones.Contains(TargetPropBone))
-				{
-					DebugInfo.TargetPairVerts.Add({TargetPropCompV_P, TargetBodyCompV_B, true,TargetBodyCompV_P, TargetPropCompV_B,});
-				}
-				else
-				{
-					DebugInfo.TargetPairVerts.Add({TargetPropCompV_P, TargetBodyCompV_B, false, FVector::ZeroVector, FVector::ZeroVector});
-				}
-				
-				DebugPairSpaces[2*OutPairIdx] = FDebugRelativeTargetPairSpace{WeightedTargetProp, TargetRelativeToBodyV+RelativePropOffset, TargetRelativeToPropV+RelativePropOffset, FeasibleRange};
-				DebugPairSpaces[2*OutPairIdx+1] = FDebugRelativeTargetPairSpace{WeightedTargetBody, TargetRelativeToBodyV+RelativeBodyOffset, TargetRelativeToPropV+RelativeBodyOffset, FeasibleRange};
+				DebugInfo.TargetPairVerts.Add({TargetPropCompV_P, TargetBodyCompV_B, true,TargetBodyCompV_P, TargetPropCompV_B,});
 			}
-#endif //WITH_EDITOR
+			else
+			{
+				DebugInfo.TargetPairVerts.Add({TargetPropCompV_P, TargetBodyCompV_B, false, FVector::ZeroVector, FVector::ZeroVector});
+			}
+			
+			DebugPairSpaces[2*OutPairIdx] = FDebugRelativeTargetPairSpace{WeightedTargetProp, TargetRelativeToBodyV+RelativePropOffset, TargetRelativeToPropV+RelativePropOffset, FeasibleRange};
+			DebugPairSpaces[2*OutPairIdx+1] = FDebugRelativeTargetPairSpace{WeightedTargetBody, TargetRelativeToBodyV+RelativeBodyOffset, TargetRelativeToPropV+RelativeBodyOffset, FeasibleRange};
 		}
+#endif //WITH_EDITOR
 	}
 }
 
@@ -714,11 +707,27 @@ const FTransform& FIKRetargetRelativeIKOp::GetTargetBoneDelta(FName TargetBoneNa
 	return CacheMapBoneSourceTargetTfm[TargetBoneName];
 }
 
+const FTransform& FIKRetargetRelativeIKOp::OffsetTransformsForSourceBones(FName SourceBoneName) const
+{
+	if (CachedNotifyInfo && CachedNotifyInfo->OffsetTransformsForBones.Contains(SourceBoneName))
+	{
+		return CachedNotifyInfo->OffsetTransformsForBones[SourceBoneName];
+	}
+
+	if (CachedPropNotifyInfo && CachedPropNotifyInfo->OffsetTransformsForBones.Contains(SourceBoneName))
+	{
+		return CachedPropNotifyInfo->OffsetTransformsForBones[SourceBoneName];
+	}
+	
+	return FTransform::Identity;
+}
+
 void FIKRetargetRelativeIKOp::ComputeSourceBodyTransform(FBodyTransform& OutTransform, FName SourceBoneName, const FTransform& GlobalTfm, double SourceScale) const
 {
-	FQuat BodyRot = GetBodyRotation(SourcePhysicsAsset, SourceBoneName);
-	FVector BodyTrans = GetBodyTranslation(SourcePhysicsAsset, SourceBoneName) * SourceScale;
-	FVector BodyScale = GetBodyOrientedScale(SourcePhysicsAsset, SourceBoneName) * SourceScale;
+	FTransform BoneOffsetTfm = OffsetTransformsForSourceBones(SourceBoneName);
+	FQuat BodyRot = BoneOffsetTfm.GetRotation();
+	FVector BodyTrans = BoneOffsetTfm.GetTranslation() * SourceScale;
+	FVector BodyScale = BoneOffsetTfm.GetScale3D() * SourceScale;
 
 	OutTransform.BoneToBody = FTransform(BodyRot.Inverse());
 	OutTransform.BodyScale = BodyScale;
@@ -732,8 +741,7 @@ void FIKRetargetRelativeIKOp::ComputeTargetBodyTransform(FBodyTransform& OutTran
 	FQuat BodyRot = GetBodyRotation(TargetPhysicsAsset, TargetBoneName);
 	FVector BodyTrans = GetBodyTranslation(TargetPhysicsAsset, TargetBoneName);
 	// TODO: Probably should always use target body scale
-	FVector BodyScale = (Settings.bTestRetargetScale) ? GetBodyOrientedScale(TargetPhysicsAsset, TargetBoneName)
-													: GetBodyOrientedScale(SourcePhysicsAsset, SourceBoneName);
+	FVector BodyScale = GetBodyOrientedScale(TargetPhysicsAsset, TargetBoneName);
 	
 	OutTransform.BoneToBody = GetTargetBoneDelta(TargetBoneName) * FTransform(BodyRot.Inverse());
 	OutTransform.BodyScale = BodyScale;
@@ -1033,6 +1041,7 @@ void FIKRetargetRelativeIKOp::AnimGraphPreUpdateMainThread(USkeletalMeshComponen
 
 		ResetCacheNotifyInfo();
 		SetupRelativeIKNotifyInfoMontage();
+		PostCachedAllNotifyInfo();
 		return;
 	}
 
@@ -1249,8 +1258,20 @@ void FIKRetargetRelativeIKOp::UpdateRelativeIKNotifyInfoAnimSeq(UAnimInstance* S
 			UpdateCachePropNotifyInfo(PropInfo);
 		}
 	}
+
+	PostCachedAllNotifyInfo();
 }
 
+void FIKRetargetRelativeIKOp::PostCachedAllNotifyInfo()
+{
+	for (int EffectIdx = 0; EffectIdx < CacheSourceEffectBones.Num(); ++EffectIdx)
+	{
+		if (CacheSourceBodyEffectVertIdx[EffectIdx].IsEmpty())
+		{
+			CacheSourceDomainBones.Add(CacheSourceEffectBones[EffectIdx]);
+		}
+	}
+}
 
 void FIKRetargetRelativeIKOp::ResetCacheNotifyInfo()
 {
@@ -1275,12 +1296,17 @@ void FIKRetargetRelativeIKOp::ResetCacheNotifyInfo()
 
 void FIKRetargetRelativeIKOp::UpdateCacheNotifyInfo(URelativeBodyBakeAnimNotify* NotifyInfo)
 {
+	if (NotifyInfo->OffsetTransformsForBones.IsEmpty())
+	{
+		return;
+	}
+	
 	CachedNotifyInfo = NotifyInfo;
 	CacheMeshPairOffset = FrameBoneTargets.Num() / 2;
-
-	int32 NumBodies = CachedNotifyInfo->BodyPairs.Num();
-	FrameBoneTargets.AddDefaulted(NumBodies);
-	FrameBoneWeights.AddZeroed(NumBodies);
+	
+	int32 NumBodyPairs = CachedNotifyInfo->BodyPairs.Num();
+	FrameBoneTargets.AddDefaulted(NumBodyPairs);
+	FrameBoneWeights.AddZeroed(NumBodyPairs);
 
 	MeshSampleRate = 30.0f;
 	if (CachedNotifyInfo->BodyPairsSampleTime.Num() >= 2)
@@ -1289,9 +1315,9 @@ void FIKRetargetRelativeIKOp::UpdateCacheNotifyInfo(URelativeBodyBakeAnimNotify*
 		MeshSampleRate = FMath::RoundToFloat(1.0f / DeltaT);
 	}
 	
-	for (int BodyIdx = 0; BodyIdx < NumBodies; ++BodyIdx)
+	for (const TPair<FName, FTransform>& BodyTfmPair : CachedNotifyInfo->OffsetTransformsForBones)
 	{
-		FName SourceBodyBone = CachedNotifyInfo->BodyPairs[BodyIdx];
+		FName SourceBodyBone = BodyTfmPair.Key;
 		if (CacheSourceEffectBones.Contains(SourceBodyBone))
 		{
 			continue;
@@ -1327,24 +1353,21 @@ void FIKRetargetRelativeIKOp::UpdateCacheNotifyInfo(URelativeBodyBakeAnimNotify*
 			CacheSourceBodyEffectVertIdx[PrimaryIdx].Add(2*OutPairIdx);
 		}
 	}
-
-	for (int EffectIdx = 0; EffectIdx < CacheSourceEffectBones.Num(); ++EffectIdx)
-	{
-		if (CacheSourceBodyEffectVertIdx[EffectIdx].IsEmpty())
-		{
-			CacheSourceDomainBones.Add(CacheSourceEffectBones[EffectIdx]);
-		}
-	}
 }
 
 void FIKRetargetRelativeIKOp::UpdateCachePropNotifyInfo(URelativePropsBakeAnimNotify* NotifyInfo)
 {
+	if (NotifyInfo->OffsetTransformsForBones.IsEmpty())
+	{
+		return;
+	}
+	
 	CachedPropNotifyInfo = NotifyInfo;
 	CachePropPairOffset = FrameBoneTargets.Num() / 2;
 
-	int32 NumBones = CachedPropNotifyInfo->PropsPairs.Num();
-	FrameBoneTargets.AddDefaulted(NumBones);
-	FrameBoneWeights.AddZeroed(NumBones);
+	int32 NumPropPairs = CachedPropNotifyInfo->PropsPairs.Num();
+	FrameBoneTargets.AddDefaulted(NumPropPairs);
+	FrameBoneWeights.AddZeroed(NumPropPairs);
 
 	PropSampleRate = 30.0f;
 	if (CachedPropNotifyInfo->PropsPairsSampleTime.Num() >= 2)
@@ -1353,9 +1376,9 @@ void FIKRetargetRelativeIKOp::UpdateCachePropNotifyInfo(URelativePropsBakeAnimNo
 		PropSampleRate = FMath::RoundToFloat(1.0f / DeltaT);
 	}
 	
-	for (int BoneIdx = 0; BoneIdx < NumBones; ++BoneIdx)
+	for (const TPair<FName, FTransform>& BodyTfmPair : CachedPropNotifyInfo->OffsetTransformsForBones)
 	{
-		FName SourceBodyBone = CachedPropNotifyInfo->PropsPairs[BoneIdx];
+		FName SourceBodyBone = BodyTfmPair.Key;
 		if (CacheSourceEffectBones.Contains(SourceBodyBone))
 		{
 			continue;
@@ -1380,6 +1403,11 @@ void FIKRetargetRelativeIKOp::UpdateCachePropNotifyInfo(URelativePropsBakeAnimNo
 		FName PropBone = CachedPropNotifyInfo->PropsPairs[2*PairIdx];
 		FName BodyBone = CachedPropNotifyInfo->PropsPairs[2*PairIdx + 1];
 
+		if (BodyBone=="Floor")
+		{
+			continue;
+		}
+
 		int32 PropIdx = CacheSourceEffectBones.Find(PropBone);
 		int32 BodyIdx = CacheSourceEffectBones.Find(BodyBone);
 
@@ -1447,29 +1475,21 @@ void FIKRetargetRelativeIKOp::DebugDraw(
 		// Draw target body pairs (Possibly with "full" retarget space display)
 		DebugDrawBodyPairs(InPDI, InComponentTransform, InComponentScale, DebugDrawInfo.TargetPairVerts);
 	}
-
+	
 	if (Settings.bDebugDrawPhysicsBodies)
 	{
-		if (Settings.SourcePhysicsAssetOverride)
-		{
-			DebugDrawBodies(InPDI, InSourceTransform, 1.0, Settings.SourcePhysicsAssetOverride, DebugDrawInfo.SourceBodyInfo);
-		}
 		if (Settings.TargetPhysicsAssetOverride)
 		{
-			DebugDrawBodies(InPDI, InComponentTransform, InComponentScale, Settings.TargetPhysicsAssetOverride, DebugDrawInfo.TargetBodyInfo);
+			// TODO: if only bake source transforms instead of PA, this source draw is meaningless
+			DebugDrawBodies(InPDI, InSourceTransform, 1.0, Settings.TargetPhysicsAssetOverride, DebugDrawInfo.SourceBodyInfo, true);
+			DebugDrawBodies(InPDI, InComponentTransform, InComponentScale, Settings.TargetPhysicsAssetOverride, DebugDrawInfo.TargetBodyInfo, false);
 		}
 	}
-
+	
 	if (Settings.bDebugDrawBodyTransforms)
 	{
-		if (Settings.SourcePhysicsAssetOverride)
-		{
-			DebugDrawBodyCoords(InPDI, InSourceTransform, 1.0, Settings.SourcePhysicsAssetOverride, DebugDrawInfo.SourceTfmInfo);
-		}
-		if (Settings.TargetPhysicsAssetOverride)
-		{
-			DebugDrawBodyCoords(InPDI, InComponentTransform, InComponentScale, Settings.TargetPhysicsAssetOverride, DebugDrawInfo.TargetTfmInfo);
-		}
+		DebugDrawBodyCoords(InPDI, InSourceTransform, 1.0, DebugDrawInfo.SourceTfmInfo);
+		DebugDrawBodyCoords(InPDI, InComponentTransform, InComponentScale,  DebugDrawInfo.TargetTfmInfo);
 	}
 
 	if (Settings.bDebugDrawRetargetVertAverages)
@@ -1573,46 +1593,66 @@ void FIKRetargetRelativeIKOp::DebugDrawPairVertRetarget(FPrimitiveDrawInterface*
 	}
 }
 
-void FIKRetargetRelativeIKOp::DebugDrawBodies(FPrimitiveDrawInterface* InPDI, const FTransform& BaseTransform, double Scale, UPhysicsAsset* PhysAsset, const TArray<FDebugDrawBodyInfo>& PhysBodies) const
+void FIKRetargetRelativeIKOp::DebugDrawBodies(FPrimitiveDrawInterface* InPDI, const FTransform& BaseTransform, double Scale, UPhysicsAsset* PhysAsset, const TArray<FDebugDrawBodyInfo>& PhysBodies, bool TargetPAForSource) const
 {
 	for (const FDebugDrawBodyInfo& BodyInfo : PhysBodies)
 	{
-		const FKShapeElem* ShapeElem = FindBodyShape(PhysAsset, BodyInfo.BodyName);
+		FKShapeElem* ShapeElem = FindBodyShape(PhysAsset, BodyInfo.BodyName);
 		if (!ShapeElem)
 		{
 			continue;
 		}
 		
 		FTransform CompTfm = BodyInfo.BoneTfm * BaseTransform;
+		if (TargetPAForSource)
+		{
+			FTransform TargetShapeTfm = ShapeElem->GetTransform();
+			FTransform SourceShapeTfm = BodyInfo.RetargetTfm;
+			SourceShapeTfm.RemoveScaling();
+			
+			CompTfm = TargetShapeTfm.Inverse() * SourceShapeTfm * CompTfm;
+			CompTfm.SetScale3D(BodyInfo.RetargetTfm.GetScale3D()/CalcReferenceShapeScale3D(ShapeElem));
+		}
 
 		// Draw physics body
-		DebugDrawPhysBody(InPDI, CompTfm, Scale * BodyInfo.Scale, ShapeElem, FLinearColor::Yellow);
+		if (DebugDrawInfo.SourceDomainBones.Contains(BodyInfo.BodyName))
+		{
+			DebugDrawPhysBody(InPDI, CompTfm, Scale * BodyInfo.Scale, ShapeElem, FLinearColor::Black);
+		}
+		else
+		{
+			DebugDrawPhysBody(InPDI, CompTfm, Scale * BodyInfo.Scale, ShapeElem, FLinearColor::Gray);
+		}
 	}
 }
 
 void FIKRetargetRelativeIKOp::DebugDrawPhysBody(FPrimitiveDrawInterface* InPDI, const FTransform& ParentTransform, double Scale, const FKShapeElem* ShapeElem, const FLinearColor& Color) const
 {
-	FTransform BodyFrame = ShapeElem->GetTransform() * ParentTransform;
+	FVector AbsBodyScale = (ParentTransform.GetScale3D() * Scale).GetAbs();
+	FTransform ParentTransformNoScale(ParentTransform);
+	ParentTransformNoScale.RemoveScaling();
+	FTransform BodyFrame = ShapeElem->GetTransform() * ParentTransformNoScale;
 
 	const FVector Translation = BodyFrame.GetLocation();
 	const FVector UnitXAxis = BodyFrame.GetUnitAxis( EAxis::X );
 	const FVector UnitYAxis = BodyFrame.GetUnitAxis( EAxis::Y );
 	const FVector UnitZAxis = BodyFrame.GetUnitAxis( EAxis::Z );
+
 	
 	switch (ShapeElem->GetShapeType())
 	{
 	case EAggCollisionShape::Box:
 		{
 			const FKBoxElem* BoxElem = static_cast<const FKBoxElem*>(ShapeElem);
-			const FVector Extent = 0.5 * FVector(BoxElem->X, BoxElem->Y, BoxElem->Z) * Scale;
+			const FVector Extent = 0.5 * FVector(BoxElem->X, BoxElem->Y, BoxElem->Z) * AbsBodyScale;
 			DrawOrientedWireBox(InPDI, Translation, UnitXAxis, UnitYAxis, UnitZAxis, Extent, Color, SDPG_Foreground);
 			return;
 		}
 	case EAggCollisionShape::Sphyl:
 		{
 			const FKSphylElem* CapsuleElem = static_cast<const FKSphylElem*>(ShapeElem);
-			const double Radius = CapsuleElem->Radius * Scale;
-			const double HalfHeight = (0.5*CapsuleElem->Length + CapsuleElem->Radius) * Scale;
+			const double Radius =CapsuleElem->GetScaledRadius(AbsBodyScale);
+			const double HalfHeight = CapsuleElem->GetScaledHalfLength(AbsBodyScale);
 			DrawWireCapsule(InPDI, Translation, UnitXAxis, UnitYAxis, UnitZAxis, Color,
 						Radius, HalfHeight,25, SDPG_Foreground, 0, 1.0);
 			return;
@@ -1620,7 +1660,7 @@ void FIKRetargetRelativeIKOp::DebugDrawPhysBody(FPrimitiveDrawInterface* InPDI,
 	case EAggCollisionShape::Sphere:
 		{
 			const FKSphereElem* SphereElem = static_cast<const FKSphereElem*>(ShapeElem);
-			const double Radius = SphereElem->Radius * Scale;
+			const double Radius = SphereElem->Radius * AbsBodyScale.GetMin();
 			DrawWireSphere(InPDI, Translation, Color, Radius, 25, SDPG_Foreground);
 			return;
 		}
@@ -1630,7 +1670,7 @@ void FIKRetargetRelativeIKOp::DebugDrawPhysBody(FPrimitiveDrawInterface* InPDI,
 }
 
 
-void FIKRetargetRelativeIKOp::DebugDrawBodyCoords(FPrimitiveDrawInterface* InPDI, const FTransform& BaseTransform, double Scale, UPhysicsAsset* PhysAsset, const TArray<FDebugBodyTfmInfo>& BodyTfms) const
+void FIKRetargetRelativeIKOp::DebugDrawBodyCoords(FPrimitiveDrawInterface* InPDI, const FTransform& BaseTransform, double Scale, const TArray<FDebugBodyTfmInfo>& BodyTfms) const
 {
 	// Stop scale from making coords too large
 	const float MaxScale = 1.0f;
@@ -1663,6 +1703,8 @@ void FIKRetargetRelativeIKOp::UpdateDebugInfo(
 	}
 	
 	FScopeLock ScopeLock(&DebugDataMutex);
+	
+	DebugDrawInfo.SourceDomainBones = LocalDebugInfo.SourceDomainBones;
 
 	DebugDrawInfo.TargetGoals = LocalDebugInfo.TargetGoals;
 	DebugDrawInfo.PairRetargetInfo = LocalDebugInfo.PairRetargetInfo;
@@ -1671,87 +1713,65 @@ void FIKRetargetRelativeIKOp::UpdateDebugInfo(
 
 	// Cache debug info for source bodies
 	DebugDrawInfo.SourceBodyInfo.Reset(CacheSourceEffectBones.Num());
-	for (FName SourceBoneName : CacheSourceEffectBones)
+	DebugDrawInfo.TargetBodyInfo.Reset(CacheSourceEffectBones.Num());
+	for (int SourceEffectBoneIndex = 0; SourceEffectBoneIndex<CacheSourceEffectBones.Num(); SourceEffectBoneIndex++)
 	{
-		if (!CacheSourceSkelIndices.Contains(SourceBoneName))
+		FName SourceBoneName = CacheSourceEffectBones[SourceEffectBoneIndex];
+		FName TargetBoneName = ApplyBodyMap(SourceBoneName);
+		
+		if (CacheSourceSkelIndices.Contains(SourceBoneName))
 		{
-			continue;
+			int32 SourcePoseIdx = CacheSourceSkelIndices[SourceBoneName];
+			// Need to apply source-scale to bodies since it's directly applied to translations only in source pose
+			FTransform SourceScaleTfm{FQuat::Identity, FVector::ZeroVector, FVector(SourceScale)};
+			DebugDrawInfo.SourceBodyInfo.Emplace(SourceBoneName, SourceScale, SourceScaleTfm*SourcePose[SourcePoseIdx], OffsetTransformsForSourceBones(SourceBoneName));
 		}
-
-		int32 SourcePoseIdx = CacheSourceSkelIndices[SourceBoneName];
-		// Need to apply source-scale to bodies since it's directly applied to translations only in source pose 
-		FTransform SourceScaleTfm{FQuat::Identity, FVector::ZeroVector, FVector(SourceScale)};
-		DebugDrawInfo.SourceBodyInfo.Emplace(SourceBoneName, SourceScale, SourceScaleTfm*SourcePose[SourcePoseIdx], FTransform::Identity);
-	}
-
-	// Cache debug info for target domain bodies
-	DebugDrawInfo.TargetBodyInfo.Reset(CacheSourceDomainBones.Num());
-	for (FName SourceBoneName : CacheSourceDomainBones)
-	{
-		FName TargetBoneName = ApplyBodyMap(SourceBoneName);
-		if (!CacheTargetSkelIndices.Contains(TargetBoneName))
+		
+		if (CacheTargetSkelIndices.Contains(TargetBoneName))
 		{
-			continue;
+			int32 TargetPoseIdx = CacheTargetSkelIndices[TargetBoneName];
+			DebugDrawInfo.TargetBodyInfo.Emplace(TargetBoneName, 1.0, TargetPose[TargetPoseIdx], FTransform::Identity);
 		}
-
-		FTransform TargetBoneDelta = GetTargetBoneDelta(TargetBoneName);
-		int32 TargetPoseIdx = CacheTargetSkelIndices[TargetBoneName];
-		DebugDrawInfo.TargetBodyInfo.Emplace(TargetBoneName, 1.0, TargetPose[TargetPoseIdx], TargetBoneDelta);
 	}
 
 	DebugDrawInfo.SourceTfmInfo.Reset(CacheSourceEffectBones.Num());
-	for (FName SourceBoneName : CacheSourceEffectBones)
-	{
-		if (!CacheSourceSkelIndices.Contains(SourceBoneName))
-		{
-			continue;
-		}
-
-		if (!FindBodyShape(SourcePhysicsAsset, SourceBoneName))
-		{
-			continue;
-		}
-
-		int32 SourcePoseIdx = CacheSourceSkelIndices[SourceBoneName];
-		
-		FBodyTransform SourceFullTfm;
-		FTransform SourceBodyRot = FTransform(GetBodyRotation(SourcePhysicsAsset, SourceBoneName));
-		ComputeSourceBodyTransform(SourceFullTfm, SourceBoneName, SourcePose[SourcePoseIdx], SourceScale);
-
-		FVector SourceBodyCenter = ApplyBodyTransform(SourceFullTfm, FVector::ZeroVector);
-		FVector SourceBodyX = ApplyBodyTransform(SourceFullTfm, SourceBodyRot.GetUnitAxis(EAxis::X));
-		FVector SourceBodyY = ApplyBodyTransform(SourceFullTfm, SourceBodyRot.GetUnitAxis(EAxis::Y));
-		FVector SourceBodyZ = ApplyBodyTransform(SourceFullTfm, SourceBodyRot.GetUnitAxis(EAxis::Z));
-		
-		DebugDrawInfo.SourceTfmInfo.Emplace(SourceBodyCenter, SourceBodyX, SourceBodyY, SourceBodyZ);
-	}
-	
 	DebugDrawInfo.TargetTfmInfo.Reset(CacheSourceEffectBones.Num());
-	for (FName SourceBoneName : CacheSourceEffectBones)
+	for (int SourceEffectBoneIndex = 0; SourceEffectBoneIndex<CacheSourceEffectBones.Num(); SourceEffectBoneIndex++)
 	{
+		FName SourceBoneName = CacheSourceEffectBones[SourceEffectBoneIndex];
 		FName TargetBoneName = ApplyBodyMap(SourceBoneName);
-		if (!CacheTargetSkelIndices.Contains(TargetBoneName))
+		
+		if (CacheSourceSkelIndices.Contains(SourceBoneName))
 		{
-			continue;
+			int32 SourcePoseIdx = CacheSourceSkelIndices[SourceBoneName];
+		
+			FBodyTransform SourceFullTfm;
+			FTransform SourceBodyRot(OffsetTransformsForSourceBones(SourceBoneName).GetRotation());
+			ComputeSourceBodyTransform(SourceFullTfm, SourceBoneName, SourcePose[SourcePoseIdx], SourceScale);
+
+			FVector SourceBodyCenter = ApplyBodyTransform(SourceFullTfm, FVector::ZeroVector);
+			FVector SourceBodyX = ApplyBodyTransform(SourceFullTfm, SourceBodyRot.GetUnitAxis(EAxis::X));
+			FVector SourceBodyY = ApplyBodyTransform(SourceFullTfm, SourceBodyRot.GetUnitAxis(EAxis::Y));
+			FVector SourceBodyZ = ApplyBodyTransform(SourceFullTfm, SourceBodyRot.GetUnitAxis(EAxis::Z));
+		
+			DebugDrawInfo.SourceTfmInfo.Emplace(SourceBodyCenter, SourceBodyX, SourceBodyY, SourceBodyZ);
 		}
 
-		if (!FindBodyShape(SourcePhysicsAsset, SourceBoneName))
+		if (CacheTargetSkelIndices.Contains(TargetBoneName))
 		{
-			continue;
-		}
-		
-		int32 TargetPoseIdx = CacheTargetSkelIndices[TargetBoneName];
+			int32 TargetPoseIdx = CacheTargetSkelIndices[TargetBoneName];
 
-		FBodyTransform TargetFullTfm;
-		FTransform SourceBodyRot = FTransform(GetBodyRotation(SourcePhysicsAsset, SourceBoneName));
-		ComputeTargetBodyTransform(TargetFullTfm, SourceBoneName, TargetPose[TargetPoseIdx]);
+			FBodyTransform TargetFullTfm;
+			FTransform SourceBodyRot(OffsetTransformsForSourceBones(SourceBoneName).GetRotation());
+			ComputeTargetBodyTransform(TargetFullTfm, SourceBoneName, TargetPose[TargetPoseIdx]);
 
-		FVector TargetBodyCenter = ApplyBodyTransform(TargetFullTfm, FVector::ZeroVector);
-		FVector TargetBodyX = ApplyBodyTransform(TargetFullTfm, SourceBodyRot.GetUnitAxis(EAxis::X));
-		FVector TargetBodyY = ApplyBodyTransform(TargetFullTfm, SourceBodyRot.GetUnitAxis(EAxis::Y));
-		FVector TargetBodyZ = ApplyBodyTransform(TargetFullTfm, SourceBodyRot.GetUnitAxis(EAxis::Z));
+			FVector TargetBodyCenter = ApplyBodyTransform(TargetFullTfm, FVector::ZeroVector);
+			FVector TargetBodyX = ApplyBodyTransform(TargetFullTfm, SourceBodyRot.GetUnitAxis(EAxis::X));
+			FVector TargetBodyY = ApplyBodyTransform(TargetFullTfm, SourceBodyRot.GetUnitAxis(EAxis::Y));
+			FVector TargetBodyZ = ApplyBodyTransform(TargetFullTfm, SourceBodyRot.GetUnitAxis(EAxis::Z));
 		
-		DebugDrawInfo.TargetTfmInfo.Emplace(TargetBodyCenter, TargetBodyX, TargetBodyY, TargetBodyZ);
+			DebugDrawInfo.TargetTfmInfo.Emplace(TargetBodyCenter, TargetBodyX, TargetBodyY, TargetBodyZ);
+		}
 	}
 }
 
diff --git a/Engine/Plugins/Experimental/Animation/RelativeIKOp/Source/RelativeIKOp/Public/RelativeIKOp.h b/Engine/Plugins/Experimental/Animation/RelativeIKOp/Source/RelativeIKOp/Public/RelativeIKOp.h
--- a/Engine/Plugins/Experimental/Animation/RelativeIKOp/Source/RelativeIKOp/Public/RelativeIKOp.h
+++ b/Engine/Plugins/Experimental/Animation/RelativeIKOp/Source/RelativeIKOp/Public/RelativeIKOp.h
@@ -84,6 +84,8 @@ struct FDebugRelativeIKDrawInfo
 	TArray<FDebugRelativeTargetPairSpace> PairRetargetInfo;
 	// Final goal points with contributions
 	TArray<FDebugDrawTargetGoal> TargetGoals;
+
+	TArray<FName> SourceDomainBones;
 };
 
 USTRUCT(BlueprintType, meta = (DisplayName = "Relative IK Settings"))
@@ -96,8 +98,8 @@ struct FIKRetargetRelativeIKOpSettings : public FIKRetargetOpSettingsBase
 	UE_API virtual void CopySettingsAtRuntime(const FIKRetargetOpSettingsBase* InSettingsToCopyFrom) override;
 
 	// Source mesh physics asset for relative body pair tests 
-	UPROPERTY(EditAnywhere, BlueprintReadWrite, Category = "Physics", meta=(ReinitializeOnEdit))
-	TObjectPtr<UPhysicsAsset> SourcePhysicsAssetOverride;
+	// UPROPERTY(EditAnywhere, BlueprintReadWrite, Category = "Physics", meta=(ReinitializeOnEdit))
+	// TObjectPtr<UPhysicsAsset> SourcePhysicsAssetOverride;
 
 	// Target mesh physics asset for retargeting relative body pair vertices
 	UPROPERTY(EditAnywhere, BlueprintReadWrite, Category = "Physics", meta=(ReinitializeOnEdit))
@@ -175,9 +177,9 @@ struct FIKRetargetRelativeIKOpSettings : public FIKRetargetOpSettingsBase
 	UPROPERTY(BlueprintReadWrite, Category = "Test Features")
 	bool bTestFeasibilityWeight = true;
 
-	// Use target (default) when computing pair alternate representations
+	/*// Use target (default) when computing pair alternate representations
 	UPROPERTY(BlueprintReadWrite, Category = "Test Features")
-	bool bTestRetargetScale = true;
+	bool bTestRetargetScale = true;*/
 
 	// Apply prop pin post-solve
 	UPROPERTY(EditAnywhere, BlueprintReadWrite, Category = "Test Features")
@@ -238,6 +240,7 @@ private:
 	void SetupRelativeIKNotifyInfoMontage();
 	void UpdateRelativeIKNotifyInfoAnimSeq(UAnimInstance* SourceAnimInstance);
 	void ResetCacheNotifyInfo();
+	void PostCachedAllNotifyInfo();
 	void UpdateCacheNotifyInfo(URelativeBodyBakeAnimNotify* NotifyInfo);
 	void UpdateCacheChainInfo(const TArray<FTransform>& TargetPose, FName BoneName);
 	void UpdateCachePropNotifyInfo(URelativePropsBakeAnimNotify* NotifyInfo);
@@ -294,7 +297,8 @@ private:
 		FVector BodyScale;
 		FTransform BodyToGlobal;
 	};
-	
+
+	const FTransform& OffsetTransformsForSourceBones(FName SourceBoneName) const;
 	void ComputeSourceBodyTransform(FBodyTransform& OutTransform, FName SourceBoneName, const FTransform& GlobalTfm, double SourceScale) const;
 	void ComputeTargetBodyTransform(FBodyTransform& OutTransform, FName SourceBoneName, const FTransform& GlobalTfm) const;
 
@@ -334,9 +338,9 @@ private:
 	void DebugDrawBodyPairs(FPrimitiveDrawInterface* InPDI, const FTransform& BaseTransform, double Scale, const TArray<FDebugDrawBodyPair>& BodyPairs) const;
 	void DebugDrawGoalContributions(FPrimitiveDrawInterface* InPDI, const FTransform& BaseTransform, double Scale, const TArray<FDebugDrawTargetGoal>& GoalInfoLis) const;
 	void DebugDrawPairVertRetarget(FPrimitiveDrawInterface* InPDI, const FTransform& BaseTransform, double Scale, const TArray<FDebugRelativeTargetPairSpace>& RetargetPairList) const;
-	void DebugDrawBodies(FPrimitiveDrawInterface* InPDI, const FTransform& BaseTransform, double Scale, UPhysicsAsset* PhysAsset, const TArray<FDebugDrawBodyInfo>& PhysBodies) const;
+	void DebugDrawBodies(FPrimitiveDrawInterface* InPDI, const FTransform& BaseTransform, double Scale, UPhysicsAsset* PhysAsset, const TArray<FDebugDrawBodyInfo>& PhysBodies, bool TargetPAForSource) const;
 	void DebugDrawPhysBody(FPrimitiveDrawInterface* InPDI, const FTransform& ParentTransform, double Scale, const FKShapeElem* ShapeElem, const FLinearColor& Color) const;
-	void DebugDrawBodyCoords(FPrimitiveDrawInterface* InPDI, const FTransform& BaseTransform, double Scale, UPhysicsAsset* PhysAsset, const TArray<FDebugBodyTfmInfo>& BodyTfms) const;
+	void DebugDrawBodyCoords(FPrimitiveDrawInterface* InPDI, const FTransform& BaseTransform, double Scale, const TArray<FDebugBodyTfmInfo>& BodyTfms) const;
 
 	void UpdateDebugInfo(double SourceScale, const TArray<FTransform>& SourcePose, const TArray<FTransform>& TargetPose, const FDebugRelativeIKDrawInfo& LocalDebugInfo);
 	void ResetDebugInfo();
@@ -348,7 +352,7 @@ private:
 #endif
 private:
 	
-	TObjectPtr<const UPhysicsAsset> SourcePhysicsAsset;
+	// TObjectPtr<const UPhysicsAsset> SourcePhysicsAsset;
 	TObjectPtr<const UPhysicsAsset> TargetPhysicsAsset;
 	
 	TObjectPtr<const UIKRigDefinition> TargetIKRig;
