RelativeIKOp: Enable bidirectional distance blend for contacts. WIP for solving dependencies between multiple moving bodies

diff --git a/Engine/Plugins/Experimental/Animation/RelativeIKOp/Source/RelativeIKOp/Private/RelativeIKOp.cpp b/Engine/Plugins/Experimental/Animation/RelativeIKOp/Source/RelativeIKOp/Private/RelativeIKOp.cpp
--- a/Engine/Plugins/Experimental/Animation/RelativeIKOp/Source/RelativeIKOp/Private/RelativeIKOp.cpp
+++ b/Engine/Plugins/Experimental/Animation/RelativeIKOp/Source/RelativeIKOp/Private/RelativeIKOp.cpp
@@ -140,6 +140,11 @@ void FIKRetargetRelativeIKOp::Run(
 		OutTargetGlobalPose[TargetBoneIdx] = InSourceGlobalPose[SourceBoneIdx];
 	}
 
+	// TODO: Should we just use initial chain lengths?
+	CacheChainLengthMap.Reset();
+	// Remove previous frame cache chain start data
+	CacheChainStartMap.Reset();
+
 	const double SourceScale = InProcessor.GetSourceScaleFactor();
 	if (CachedNotifyInfo)
 	{
@@ -226,10 +231,6 @@ void FIKRetargetRelativeIKOp::ComputeFramePairSpaces(
 {
 	int32 PairCount = CachedNotifyInfo->bBodyPairsIsParentDominates.Num();
 	const double InvSourceScale = (Settings.bIgnoreSourceScale) ? 1.0 / SourceScale : 1.0;
-
-	// TODO: Should we just use initial chain lengths?
-	CacheChainLengthMap.Reset();
-	CacheChainStartMap.Reset();
 	
 	const double RetargetContactAlpha = FMath::Clamp(Settings.RetargetContactAlpha, 0.0, 1.0);
 	const double RetargetSpringAlpha = FMath::Clamp(Settings.RetargetSpringAlpha, 0.0, 1.0);
@@ -240,8 +241,11 @@ void FIKRetargetRelativeIKOp::ComputeFramePairSpaces(
 		const int32 OutPairIdx = PairIdx + PairOffset;
 		FrameBoneWeights[2*OutPairIdx] = 0;
 		FrameBoneWeights[2*OutPairIdx + 1] = 0;
-		FrameBoneTargets[2*OutPairIdx] = FVector::ZeroVector;
-		FrameBoneTargets[2*OutPairIdx + 1] = FVector::ZeroVector;
+		FrameBoneOffsets[2*OutPairIdx] = FVector::ZeroVector;
+		FrameBoneOffsets[2*OutPairIdx + 1] = FVector::ZeroVector;
+		//TODO: Remove FrameBonePositions and figure out a quick access from effect vert indices to bone positions
+		FrameBonePositions[2*OutPairIdx] = FVector::ZeroVector;
+		FrameBonePositions[2*OutPairIdx + 1] = FVector::ZeroVector;
 		
 		FName SourcePrimaryBone = CachedNotifyInfo->BodyPairs[2*PairIdx];
 		FName SourceOtherBone = CachedNotifyInfo->BodyPairs[2*PairIdx + 1];
@@ -269,6 +273,8 @@ void FIKRetargetRelativeIKOp::ComputeFramePairSpaces(
 			continue;
 		}
 
+		bool bPrimaryMovable = CachedNotifyInfo->bBodyPairsIsParentDominates[PairIdx];
+
 		int32 SourcePrimaryPoseIdx = CacheSourceSkelIndices[SourcePrimaryBone];
 		int32 SourceOtherPoseIdx = CacheSourceSkelIndices[SourceOtherBone];
 
@@ -325,42 +331,39 @@ void FIKRetargetRelativeIKOp::ComputeFramePairSpaces(
 		// Compute "other" body vert represented in component space relative to primary (_P) and other (_O) bodies
 		FVector TargetOtherCompV_O = ApplyBodyTransform(TargetOtherTfm, SourceBoneBakeOtherV);
 		FVector TargetOtherCompV_P = ApplyBodyTransform(TargetPrimaryTfm, SourceOtherV_P);
-
-		// Get offsets from target verts (in their own body rep) to associated bone
-		FVector TargetPrimaryBoneOffset = TargetPose[TargetPrimaryPoseIdx].GetTranslation() - TargetPrimaryCompV_P;
-		FVector TargetOtherBoneOffset = TargetPose[TargetOtherPoseIdx].GetTranslation() - TargetOtherCompV_O;
-
-		// double SpringAlpha = (Settings.bTestDistContactAlpha) ? Weight : RetargetSpringAlpha;
+		
 		// These are the target midpoint positions used for spring constraints in PhysicalIK
-		FVector TargetRelativeToPrimaryV = FMath::Lerp(TargetPrimaryCompV_P,TargetOtherCompV_P, RetargetSpringAlpha);
-		FVector TargetRelativeToOtherV = FMath::Lerp(TargetPrimaryCompV_O,TargetOtherCompV_O, RetargetSpringAlpha);
+		FVector TargetSpring_P = FMath::Lerp(TargetPrimaryCompV_P,TargetOtherCompV_P, RetargetSpringAlpha);
+		FVector TargetSpring_O = FMath::Lerp(TargetPrimaryCompV_O, TargetOtherCompV_O, RetargetSpringAlpha);
 
-		// double TargetVDist = FVector::Distance(TargetRelativeToPrimaryV, TargetRelativeToOtherV);
-		// double TargetDRat = FMath::Min(1.0, 2.0 * FMath::Min(SourceVDist/TargetVDist, TargetVDist/SourceVDist));
-		
-		FVector RelativePrimaryOffset = TargetPrimaryCompV_P - TargetRelativeToPrimaryV;
-		FVector RelativeOtherOffset = TargetOtherCompV_O - TargetRelativeToOtherV;
+		// Offset vector (Other -> Primary) to Lerp for dist alpha
+		FVector TargetOffsetVec = TargetSpring_P - TargetSpring_O;
 
-		// TODO: Should Eventually remove testing setup (keep dist weighting or drop it)
-		bool bDistanceWeighting = CachedNotifyInfo->bBodyPairsIsParentDominates[PairIdx] && Settings.bTestDistContactAlpha;
+		FVector TargetPrimaryBoneLoc = TargetPose[TargetPrimaryPoseIdx].GetLocation();
+		FVector TargetOtherBoneLoc = TargetPose[TargetOtherPoseIdx].GetLocation();
 		
-		FDoubleInterval FeasibleRange(0.0, 1.0);
+		FDoubleInterval FeasibleRange_O(0.0, 1.0);
+		FDoubleInterval FeasibleRange_P(0.0, 1.0);
 		UpdateCacheChainInfo(TargetPose, TargetOtherBone);
-		if (bDistanceWeighting)
+		UpdateCacheChainInfo(TargetPose, TargetPrimaryBone);
+		ComputeFeasibilityRange(FeasibleRange_O, TargetOtherBone, TargetOtherBoneLoc, TargetOffsetVec);
+		if (bPrimaryMovable)
 		{
-			// TODO: Experiment with bi-directional distance weighting when Primary/Other can both move
-			ComputeFeasibilityRange(FeasibleRange, TargetOtherBone, RelativeOtherOffset + TargetOtherBoneOffset, TargetRelativeToOtherV, TargetRelativeToPrimaryV);
+			ComputeFeasibilityRange(FeasibleRange_P, TargetPrimaryBone, TargetPrimaryBoneLoc, -TargetOffsetVec);
 		}
 
-		double FeasibilityWeight = FMath::Lerp(FeasibleRange.Min, FeasibleRange.Max, LinearWeight);
-		double ContactAlpha = (bDistanceWeighting) ? FeasibilityWeight : RetargetContactAlpha;
-		FVector WeightedTarget = FMath::Lerp(TargetRelativeToOtherV, TargetRelativeToPrimaryV, ContactAlpha);
+		double ContactBlendWeight = (Settings.bBidirectionalDistanceContactBlend) ? LinearWeight : 0.5;
+		double FeasibilityWeight_O = FMath::Lerp(FeasibleRange_O.Min, FeasibleRange_O.Max, ContactBlendWeight);
+		FVector WeightedTargetOffset_O = TargetOffsetVec * FeasibilityWeight_O;
 
-		FVector WeightedTargetPrimary = WeightedTarget + RelativePrimaryOffset;
-		FVector WeightedTargetOther = WeightedTarget + RelativeOtherOffset;
+		double FeasibilityWeight_P = FMath::Lerp(FeasibleRange_P.Min, FeasibleRange_P.Max, ContactBlendWeight);
+		FVector WeightedTargetOffset_P = -TargetOffsetVec * FeasibilityWeight_P;
 		
-		FrameBoneTargets[2*OutPairIdx] = WeightedTargetPrimary + TargetPrimaryBoneOffset;
-		FrameBoneTargets[2*OutPairIdx + 1] = WeightedTargetOther + TargetOtherBoneOffset;
+		FrameBoneOffsets[2*OutPairIdx] = WeightedTargetOffset_P;
+		FrameBoneOffsets[2*OutPairIdx + 1] = WeightedTargetOffset_O;
+
+		FrameBonePositions[2*OutPairIdx] = TargetPrimaryBoneLoc;
+		FrameBonePositions[2*OutPairIdx + 1] = TargetOtherBoneLoc;
 
 		FrameBoneWeights[2*OutPairIdx] = Weight;
 		FrameBoneWeights[2*OutPairIdx + 1] = Weight;
@@ -379,8 +382,8 @@ void FIKRetargetRelativeIKOp::ComputeFramePairSpaces(
 				DebugInfo.TargetPairVerts.Add({TargetPrimaryCompV_P, TargetOtherCompV_O, false, FVector::ZeroVector, FVector::ZeroVector});
 			}
 			
-			DebugPairSpaces[2*OutPairIdx] = FDebugRelativeTargetPairSpace{WeightedTargetPrimary, TargetRelativeToOtherV+RelativePrimaryOffset, TargetRelativeToPrimaryV+RelativePrimaryOffset, FeasibleRange};
-			DebugPairSpaces[2*OutPairIdx+1] = FDebugRelativeTargetPairSpace{WeightedTargetOther, TargetRelativeToOtherV+RelativeOtherOffset, TargetRelativeToPrimaryV+RelativeOtherOffset, FeasibleRange};
+			DebugPairSpaces[2*OutPairIdx] = FDebugRelativeTargetPairSpace{TargetPrimaryCompV_P, TargetPrimaryCompV_P-TargetOffsetVec, FeasibilityWeight_P, FeasibleRange_P};
+			DebugPairSpaces[2*OutPairIdx+1] = FDebugRelativeTargetPairSpace{TargetOtherCompV_O, TargetOtherCompV_O+TargetOffsetVec, FeasibilityWeight_O, FeasibleRange_O};
 		}
 #endif //WITH_EDITOR
 	}
@@ -404,9 +407,6 @@ void FIKRetargetRelativeIKOp::ComputeFramePropSpaces(
 	FTransform SourceScaleTfm(FQuat::Identity, FVector::ZeroVector, FVector(SourceScale));
 	const double InvSourceScale = (Settings.bIgnoreSourceScale) ? 1.0 / SourceScale : 1.0;
 	
-	// CacheChainLengthMap.Reset();
-	// CacheChainStartMap.Reset();
-	
 	const double RetargetContactAlpha = FMath::Clamp(Settings.RetargetContactAlpha, 0.0, 1.0);
 	const double RetargetSpringAlpha = FMath::Clamp(Settings.RetargetSpringAlpha, 0.0, 1.0);
 	// Get pair target body verts and weights
@@ -416,8 +416,11 @@ void FIKRetargetRelativeIKOp::ComputeFramePropSpaces(
 		const int32 OutPairIdx = PairIdx + PairOffset;
 		FrameBoneWeights[2*OutPairIdx] = 0;
 		FrameBoneWeights[2*OutPairIdx + 1] = 0;
-		FrameBoneTargets[2*OutPairIdx] = FVector::ZeroVector;
-		FrameBoneTargets[2*OutPairIdx + 1] = FVector::ZeroVector;
+		FrameBoneOffsets[2*OutPairIdx] = FVector::ZeroVector;
+		FrameBoneOffsets[2*OutPairIdx + 1] = FVector::ZeroVector;
+		//TODO: Remove FrameBonePositions and figure out a quick access from effect vert indices to bone positions
+		FrameBonePositions[2*OutPairIdx] = FVector::ZeroVector;
+		FrameBonePositions[2*OutPairIdx + 1] = FVector::ZeroVector;
 
 		FName SourcePropBone = CachedPropNotifyInfo->PropsPairs[2 * PairIdx];
 		FName TargetPropBone = ApplyBodyMap(SourcePropBone);
@@ -458,6 +461,8 @@ void FIKRetargetRelativeIKOp::ComputeFramePropSpaces(
 		const FVector3f& LocalPropVert = FramePointView[2*PairIdx];
 		const FVector3f& LocalBodyVert = FramePointView[2*PairIdx+1];
 
+		bool bBodyMovable = CacheSourceDomainBones.Contains(SourceBodyBone);
+
 		// TODO: Clean this up if we bake in rot invariant space
 		FVector SourceBoneBakePropV = (PropOffsetTfm).TransformPosition(FVector(LocalPropVert));
 		FVector SourceBoneBakeBodyV = (SourceBodyRot).TransformPosition(FVector(LocalBodyVert));
@@ -502,39 +507,37 @@ void FIKRetargetRelativeIKOp::ComputeFramePropSpaces(
 		FVector TargetBodyCompV_B = ApplyBodyTransform(TargetBodyTfm, SourceBoneBakeBodyV);
 		FVector TargetBodyCompV_P = TargetPropTfm.TransformPosition(SourceBodyV_P);
 
-		// Get offsets from target verts (in their own body rep) to associated bone
-		FVector TargetPropBoneOffset = TargetPose[TargetPropPoseIdx].GetTranslation() - TargetPropCompV_P;
-		FVector TargetBodyBoneOffset = TargetPose[TargetBodyPoseIdx].GetTranslation() - TargetBodyCompV_B;
 
-		// double SpringAlpha = (Settings.bTestDistContactAlpha) ? Weight : RetargetSpringAlpha;
-		FVector TargetRelativeToPropV = FMath::Lerp(TargetPropCompV_P,TargetBodyCompV_P, RetargetSpringAlpha);
-		FVector TargetRelativeToBodyV = FMath::Lerp(TargetPropCompV_B,TargetBodyCompV_B, RetargetSpringAlpha);
+		// These are the target midpoint positions used for spring constraints in PhysicalIK
+		FVector TargetSpring_P = FMath::Lerp(TargetPropCompV_P,TargetBodyCompV_P, RetargetSpringAlpha);
+		FVector TargetSpring_B = FMath::Lerp(TargetPropCompV_B, TargetBodyCompV_B, RetargetSpringAlpha);
 
-		// double TargetVDist = FVector::Distance(TargetRelativeToPrimaryV, TargetRelativeToOtherV);
-		// double TargetDRat = FMath::Min(1.0, 2.0 * FMath::Min(SourceVDist/TargetVDist, TargetVDist/SourceVDist));
-		
-		FVector RelativePropOffset = TargetPropCompV_P - TargetRelativeToPropV;
-		FVector RelativeBodyOffset = TargetBodyCompV_B - TargetRelativeToBodyV;
-		
-		// bool bDistanceWeighting = false;
-		
-		FDoubleInterval FeasibleRange(0.0, 1.0);
-		// UpdateCacheChainInfo(TargetPose, TargetBodyBone);
-		// if (bDistanceWeighting)
-		// {
-		// 	ComputeFeasibilityRange(FeasibleRange, TargetBodyBone, RelativeBodyOffset + TargetBodyBoneOffset, TargetRelativeToBodyV, TargetRelativeToPropV);
-		// }
-
-		double FeasibilityWeight = FMath::Lerp(FeasibleRange.Min, FeasibleRange.Max, LinearWeight);
-		// double ContactAlpha = (bDistanceWeighting) ? FeasibilityWeight : RetargetContactAlpha;
-		double ContactAlpha = RetargetContactAlpha;
-		FVector WeightedTarget = FMath::Lerp(TargetRelativeToBodyV, TargetRelativeToPropV, ContactAlpha);
-
-		FVector WeightedTargetProp = WeightedTarget + RelativePropOffset;
-		FVector WeightedTargetBody = WeightedTarget + RelativeBodyOffset;
-		
-		FrameBoneTargets[2*OutPairIdx] = WeightedTargetProp + TargetPropBoneOffset;
-		FrameBoneTargets[2*OutPairIdx + 1] = WeightedTargetBody + TargetBodyBoneOffset;
+		// Offset vector (Body -> Prop) to Lerp for dist alpha
+		FVector TargetOffsetVec = TargetSpring_P - TargetSpring_B;
+
+		FVector TargetPropBoneLoc = TargetPose[TargetPropPoseIdx].GetLocation();
+		FVector TargetBodyBoneLoc = TargetPose[TargetBodyPoseIdx].GetLocation();
+
+		FDoubleInterval FeasibleRange_P(0.0, 1.0);
+		FDoubleInterval FeasibleRange_B(0.0, 1.0);
+		UpdateCacheChainInfo(TargetPose, TargetBodyBone);
+		if (bBodyMovable)
+		{
+			ComputeFeasibilityRange(FeasibleRange_B, TargetBodyBone, TargetBodyBoneLoc, TargetOffsetVec);
+		}
+
+		double ContactBlendWeight = (Settings.bBidirectionalDistanceContactBlend) ? LinearWeight : 0.5;
+		double FeasibilityWeight_B =  FMath::Lerp(FeasibleRange_B.Min, FeasibleRange_B.Max, ContactBlendWeight);
+		FVector WeightedTargetOffset_B = TargetOffsetVec * FeasibilityWeight_B;
+
+		double FeasibilityWeight_P = FMath::Lerp(FeasibleRange_P.Min, FeasibleRange_P.Max, ContactBlendWeight);
+		FVector WeightedTargetOffset_P = -TargetOffsetVec * FeasibilityWeight_P;
+
+		FrameBoneOffsets[2*OutPairIdx] = WeightedTargetOffset_P;
+		FrameBoneOffsets[2*OutPairIdx + 1] = WeightedTargetOffset_B;
+
+		FrameBonePositions[2*OutPairIdx] = TargetPropBoneLoc;
+		FrameBonePositions[2*OutPairIdx + 1] = TargetBodyBoneLoc;
 
 		FrameBoneWeights[2*OutPairIdx] = Weight;
 		FrameBoneWeights[2*OutPairIdx + 1] = Weight;
@@ -553,13 +556,26 @@ void FIKRetargetRelativeIKOp::ComputeFramePropSpaces(
 				DebugInfo.TargetPairVerts.Add({TargetPropCompV_P, TargetBodyCompV_B, false, FVector::ZeroVector, FVector::ZeroVector});
 			}
 			
-			DebugPairSpaces[2*OutPairIdx] = FDebugRelativeTargetPairSpace{WeightedTargetProp, TargetRelativeToBodyV+RelativePropOffset, TargetRelativeToPropV+RelativePropOffset, FeasibleRange};
-			DebugPairSpaces[2*OutPairIdx+1] = FDebugRelativeTargetPairSpace{WeightedTargetBody, TargetRelativeToBodyV+RelativeBodyOffset, TargetRelativeToPropV+RelativeBodyOffset, FeasibleRange};
+			DebugPairSpaces[2*OutPairIdx] = FDebugRelativeTargetPairSpace{TargetPropCompV_P, TargetPropCompV_P-TargetOffsetVec, FeasibilityWeight_P, FeasibleRange_P};
+			DebugPairSpaces[2*OutPairIdx+1] = FDebugRelativeTargetPairSpace{TargetBodyCompV_B, TargetBodyCompV_B+TargetOffsetVec, FeasibilityWeight_B, FeasibleRange_B};
 		}
 #endif //WITH_EDITOR
 	}
 }
 
+void FIKRetargetRelativeIKOp::DropPropDomainBodyEffects(int32 PairOffset)
+{
+	// TODO: This should be handled by baking necessary info in prop/body notifies
+	for (int32 Index=0; Index < CacheSourceEffectBones.Num(); ++Index)
+	{
+		const FName SourceBoneName = CacheSourceEffectBones[Index]; 
+		if (CacheSourceDomainBones.Contains(SourceBoneName))
+		{
+			CacheSourceBodyEffectVertIdx[Index].Reset();
+		}
+	}
+}
+
 void FIKRetargetRelativeIKOp::UpdateTargetIKGoal(
 	TArray<FIKRigGoal>& OutIKGoals,
 	FName TargetBoneName,
@@ -652,14 +668,15 @@ double FIKRetargetRelativeIKOp::ComputeTargetWeightedLocation(
 		{
 			continue;
 		}
-		
-		OutTargetLocation += FrameBoneWeights[Index] * FrameBoneTargets[Index];
+
+		FVector TargetLoc = FrameBoneOffsets[Index] + FrameBonePositions[Index];
+		OutTargetLocation += FrameBoneWeights[Index] * TargetLoc;
 		TotalWeight += FrameBoneWeights[Index];
 
 #if WITH_EDITOR
 		if (Settings.bDebugDraw)
 		{
-			DebugInfo.TargetGoals.Last().PairTargets.Add(FrameBoneTargets[Index]);
+			DebugInfo.TargetGoals.Last().PairTargets.Add(TargetLoc);
 			DebugInfo.PairRetargetInfo.Emplace(DebugPairSpaces[Index]);
 		}
 #endif //WITH_EDITOR
@@ -827,24 +844,22 @@ TConstArrayView<FVector3f> FIKRetargetRelativeIKOp::ApplyTemporalSmoothing(float
 	return SmoothedPoints;
 }
 
-void FIKRetargetRelativeIKOp::ComputeFeasibilityRange(FDoubleInterval& OutFeasibleRange, FName BoneName, FVector BoneChainOffset, FVector TargetRelativeToBoneV, FVector TargetRelativeToOpposeBoneV)
+void FIKRetargetRelativeIKOp::ComputeFeasibilityRange(FDoubleInterval& OutFeasibleRange, FName BoneName, const FVector& BoneLoc, const FVector& TargetOffsetVec)
 {
 	OutFeasibleRange = {0.0, 1.0};
-	if (!Settings.bTestFeasibilityWeight)
-	{
-		return;
-	}
-
 	if (!CacheChainLengthMap.Contains(BoneName))
 	{
 		return;
 	}
 	
-	FVector OffsetChainStart = CacheChainStartMap[BoneName] - BoneChainOffset;
+	FVector ChainRootLoc = CacheChainStartMap[BoneName];
 	double CheckRadius = FMath::Max(CacheChainLengthMap[BoneName] + Settings.FeasibilityLengthBias, 0.0);
+
+	FVector StartBoneLoc = BoneLoc;
+	FVector EndBoneLoc = BoneLoc + TargetOffsetVec;
 			
-	double ChainDistStart = FVector::Distance(TargetRelativeToBoneV, OffsetChainStart);
-	double ChainDistEnd = FVector::Distance(TargetRelativeToOpposeBoneV, OffsetChainStart);
+	double ChainDistStart = FVector::Distance(StartBoneLoc, ChainRootLoc);
+	double ChainDistEnd = FVector::Distance(EndBoneLoc, ChainRootLoc);
 
 	FDoubleInterval IntersectRange;
 	if (FMath::Max(ChainDistStart, ChainDistEnd) <= CheckRadius)
@@ -856,7 +871,7 @@ void FIKRetargetRelativeIKOp::ComputeFeasibilityRange(FDoubleInterval& OutFeasib
 
 	// If valid intersect: Sphere intersection range (return clamped to line segment) 
 	// No valid intersect: Closest point to sphere (return clamped to line segment)
-	CalcLineSphereIntersect(IntersectRange, OffsetChainStart, CheckRadius, TargetRelativeToBoneV, TargetRelativeToOpposeBoneV);
+	CalcLineSphereIntersect(IntersectRange, ChainRootLoc, CheckRadius, StartBoneLoc, EndBoneLoc);
 	OutFeasibleRange = {FMath::Clamp(IntersectRange.Min, 0.0,1.0), FMath::Clamp(IntersectRange.Max, 0.0,1.0)};
 }
 
@@ -1289,7 +1304,8 @@ void FIKRetargetRelativeIKOp::ResetCacheNotifyInfo()
 	CacheSourceEffectBones.Reset();
 	CacheSourceDomainBones.Reset();
 	CacheBodyBones.Reset();
-	FrameBoneTargets.Reset();
+	FrameBoneOffsets.Reset();
+	FrameBonePositions.Reset();
 	FrameBoneWeights.Reset();
 	CachePinSourcePropBones.Reset();
 }
@@ -1302,11 +1318,12 @@ void FIKRetargetRelativeIKOp::UpdateCacheNotifyInfo(URelativeBodyBakeAnimNotify*
 	}
 	
 	CachedNotifyInfo = NotifyInfo;
-	CacheMeshPairOffset = FrameBoneTargets.Num() / 2;
-	
-	int32 NumBodyPairs = CachedNotifyInfo->BodyPairs.Num();
-	FrameBoneTargets.AddDefaulted(NumBodyPairs);
-	FrameBoneWeights.AddZeroed(NumBodyPairs);
+	CacheMeshPairOffset = FrameBoneOffsets.Num() / 2;
+
+	int32 NumBodies = CachedNotifyInfo->BodyPairs.Num();
+	FrameBoneOffsets.AddZeroed(NumBodies);
+	FrameBonePositions.AddZeroed(NumBodies);
+	FrameBoneWeights.AddZeroed(NumBodies);
 
 	MeshSampleRate = 30.0f;
 	if (CachedNotifyInfo->BodyPairsSampleTime.Num() >= 2)
@@ -1363,11 +1380,12 @@ void FIKRetargetRelativeIKOp::UpdateCachePropNotifyInfo(URelativePropsBakeAnimNo
 	}
 	
 	CachedPropNotifyInfo = NotifyInfo;
-	CachePropPairOffset = FrameBoneTargets.Num() / 2;
+	CachePropPairOffset = FrameBoneOffsets.Num() / 2;
 
-	int32 NumPropPairs = CachedPropNotifyInfo->PropsPairs.Num();
-	FrameBoneTargets.AddDefaulted(NumPropPairs);
-	FrameBoneWeights.AddZeroed(NumPropPairs);
+	int32 NumBones = CachedPropNotifyInfo->PropsPairs.Num();
+	FrameBoneOffsets.AddZeroed(NumBones);
+	FrameBonePositions.AddZeroed(NumBones);
+	FrameBoneWeights.AddZeroed(NumBones);
 
 	PropSampleRate = 30.0f;
 	if (CachedPropNotifyInfo->PropsPairsSampleTime.Num() >= 2)
@@ -1398,19 +1416,18 @@ void FIKRetargetRelativeIKOp::UpdateCachePropNotifyInfo(URelativePropsBakeAnimNo
 	for (int PairIdx = 0; PairIdx < PairCount; ++PairIdx)
 	{
 		int32 OutPairIdx = PairIdx + CachePropPairOffset;
-		bool bPrimaryFixed = false;
 		
 		FName PropBone = CachedPropNotifyInfo->PropsPairs[2*PairIdx];
 		FName BodyBone = CachedPropNotifyInfo->PropsPairs[2*PairIdx + 1];
 
-		if (BodyBone=="Floor")
+		int32 PropIdx = CacheSourceEffectBones.Find(PropBone);
+		int32 BodyIdx = CacheSourceEffectBones.Find(BodyBone);
+
+		if (PropIdx == INDEX_NONE || BodyIdx == INDEX_NONE)
 		{
 			continue;
 		}
 
-		int32 PropIdx = CacheSourceEffectBones.Find(PropBone);
-		int32 BodyIdx = CacheSourceEffectBones.Find(BodyBone);
-
 		// TODO: Only add IK goal bones here!
 		CacheSourceBodyEffectVertIdx[BodyIdx].Add(2*OutPairIdx+1);
 		// Prop attach bones always allowed to move
@@ -1579,7 +1596,8 @@ void FIKRetargetRelativeIKOp::DebugDrawPairVertRetarget(FPrimitiveDrawInterface*
 		// DrawDashedLine(InPDI, VBone, V2, FLinearColor::Yellow, 1, SDPG_Foreground);
 		DrawDashedLine(InPDI, V1, V2, FLinearColor::Black, 1, SDPG_Foreground);
 
-		FTransform VAvgTfm = FTransform(Info.PairTarget) * BaseTransform;
+		FVector PairTargetV = FMath::Lerp(Info.PairRangeStart, Info.PairRangeEnd, Info.TargetAlpha);
+		FTransform VAvgTfm = FTransform(PairTargetV) * BaseTransform;
 		IKRigDebugRendering::DrawWireCube(InPDI, VAvgTfm, FLinearColor::Gray,MarkerSize,MarkerSize);
 
 		FVector FeasibleMinV =  FMath::Lerp(Info.PairRangeStart, Info.PairRangeEnd, Info.FeasibleRange.Min);
diff --git a/Engine/Plugins/Experimental/Animation/RelativeIKOp/Source/RelativeIKOp/Public/RelativeIKOp.h b/Engine/Plugins/Experimental/Animation/RelativeIKOp/Source/RelativeIKOp/Public/RelativeIKOp.h
--- a/Engine/Plugins/Experimental/Animation/RelativeIKOp/Source/RelativeIKOp/Public/RelativeIKOp.h
+++ b/Engine/Plugins/Experimental/Animation/RelativeIKOp/Source/RelativeIKOp/Public/RelativeIKOp.h
@@ -22,11 +22,10 @@ class URelativePropsBakeAnimNotify;
 // Debug info for individual body pair contributions
 struct FDebugRelativeTargetPairSpace
 {
-	FVector PairTarget;
-	
 	FVector PairRangeStart;
 	FVector PairRangeEnd;
-	
+
+	double TargetAlpha;
 	FDoubleInterval FeasibleRange;
 };
 
@@ -169,17 +168,17 @@ struct FIKRetargetRelativeIKOpSettings : public FIKRetargetOpSettingsBase
 	UPROPERTY(EditAnywhere, BlueprintReadWrite, Category = "Debug")
 	bool bDryRun = false;
 
-	// Use distance to pair contact to blend towards direct contact
-	UPROPERTY(BlueprintReadWrite, Category = "Test Features")
-	bool bTestDistContactAlpha = true;
+	// // Distance blend contact bodies towards domain bodies
+	// UPROPERTY(BlueprintReadWrite, Category = "Test Features")
+	// bool bTestDistContactAlpha = true;
+	//
+	// // Reduce distance contact blending based on feasibility length of target bone chain
+	// UPROPERTY(BlueprintReadWrite, Category = "Test Features")
+	// bool bTestFeasibilityWeight = true;
 
-	// Dial back distance weighting based on feasibility length of target bone chain
-	UPROPERTY(BlueprintReadWrite, Category = "Test Features")
-	bool bTestFeasibilityWeight = true;
-
-	/*// Use target (default) when computing pair alternate representations
-	UPROPERTY(BlueprintReadWrite, Category = "Test Features")
-	bool bTestRetargetScale = true;*/
+	// Distance blend body pairs towards each contact position if both can move
+	UPROPERTY(EditAnywhere, BlueprintReadWrite, Category = "Test Features")
+	bool bBidirectionalDistanceContactBlend = true;
 
 	// Apply prop pin post-solve
 	UPROPERTY(EditAnywhere, BlueprintReadWrite, Category = "Test Features")
@@ -245,6 +244,8 @@ private:
 	void UpdateCacheChainInfo(const TArray<FTransform>& TargetPose, FName BoneName);
 	void UpdateCachePropNotifyInfo(URelativePropsBakeAnimNotify* NotifyInfo);
 
+	void DropPropDomainBodyEffects(int32 PairOffset);
+	
 	void ComputeFramePairSpaces(
 		double SourceScale,
 		const TArray<FTransform>& SourcePose,
@@ -305,7 +306,7 @@ private:
 	static FVector ApplyBodyTransform(const FBodyTransform& Transform, const FVector& LocalPos);
 	static FVector InverseBodyTransform(const FBodyTransform& Transform, const FVector& GlobalPos);
 	
-	void ComputeFeasibilityRange(FDoubleInterval& OutFeasibleRange, FName BoneName, FVector BoneChainOffset, FVector TargetRelativeToBoneV, FVector TargetRelativeToOpposeBoneV);
+	void ComputeFeasibilityRange(FDoubleInterval& OutFeasibleRange, FName BoneName, const FVector& BoneLoc, const FVector& TargetOffsetVec);
 	TConstArrayView<FVector3f> ApplyTemporalSmoothing(float Time, float SampleRate, int32 NumSamples, int32 NumBodies, const TArray<FVector3f>& PairLocalVerts);
 
 	FVector GoalLocBlendCompSpace(const FIKRigGoal* Goal, const FTransform& BoneTfm) const;
@@ -380,7 +381,8 @@ private:
 	TArray<FName> CacheSourceDomainBones;
 	TSet<FName> CacheBodyBones;
 	TArray<double> FrameBoneWeights;
-	TArray<FVector> FrameBoneTargets;
+	TArray<FVector> FrameBonePositions;
+	TArray<FVector> FrameBoneOffsets;
 
 	// Cache skel and IK goal info
 	TArray<FTransform> CacheSourceBoneInitTfm;
