Relative IK: rotation invariant baking update

diff --git a/Engine/Plugins/Experimental/Animation/RelativeIKOp/Source/RelativeBodyAnimUtils/Private/RelativeBodyAnimModifier.cpp b/Engine/Plugins/Experimental/Animation/RelativeIKOp/Source/RelativeBodyAnimUtils/Private/RelativeBodyAnimModifier.cpp
--- a/Engine/Plugins/Experimental/Animation/RelativeIKOp/Source/RelativeBodyAnimUtils/Private/RelativeBodyAnimModifier.cpp
+++ b/Engine/Plugins/Experimental/Animation/RelativeIKOp/Source/RelativeBodyAnimUtils/Private/RelativeBodyAnimModifier.cpp
@@ -274,8 +274,6 @@ void URelativeBodyAnimModifier::OnApply_Implementation(UAnimSequence* InAnimatio
 			}
 
 			TArray<FName> ContactedBodyPairs;
-			TArray<FVector3f> ContactPoints;
-			ContactPoints.Empty();
 			if (BakedAnimNotifyInfo)
 			{
 				BakedAnimNotifyInfo->BodyPairsSampleTime[SampleIndex] = SampleTime;
@@ -347,8 +345,6 @@ void URelativeBodyAnimModifier::OnApply_Implementation(UAnimSequence* InAnimatio
 
 						FVector3f Loc1 = FRelativeBodyAnimUtils::CalcVertLocationInUnitBody(ContactPoint1, BoneName1, RefSkeleton, CacheToLocals, CachedBodySourceData.SourceRetargetGlobalPose, PhysicsAsset);
 						FVector3f Loc2 = FRelativeBodyAnimUtils::CalcVertLocationInUnitBody(ContactPoint2, BoneName2, RefSkeleton, CacheToLocals, CachedBodySourceData.SourceRetargetGlobalPose, PhysicsAsset);
-						ContactPoints.Add(Loc1);
-						ContactPoints.Add(Loc2);
 						
 						// Use your custom notify class
 						if (!bUseDenseNotifyInfo)
@@ -371,9 +367,12 @@ void URelativeBodyAnimModifier::OnApply_Implementation(UAnimSequence* InAnimatio
 						}
 						else if (BakedAnimNotifyInfo)
 						{
+							FTransform3f SourceBodyRot1(FTransform(BakedAnimNotifyInfo->OffsetTransformsForBones[BoneName1].GetRotation()));
+							FTransform3f SourceBodyRot2(FTransform(BakedAnimNotifyInfo->OffsetTransformsForBones[BoneName2].GetRotation()));
+							
 							size_t offset = SampleIndex*BakedAnimNotifyInfo->BodyPairs.Num();
-							BakedAnimNotifyInfo->BodyPairsLocalReference[offset+count*2+0] = Loc1;
-							BakedAnimNotifyInfo->BodyPairsLocalReference[offset+count*2+1] = Loc2;
+							BakedAnimNotifyInfo->BodyPairsLocalReference[offset+count*2+0] = SourceBodyRot1.TransformPosition(Loc1);
+							BakedAnimNotifyInfo->BodyPairsLocalReference[offset+count*2+1] = SourceBodyRot2.TransformPosition(Loc2);
 						}
 					}
 					count++;
diff --git a/Engine/Plugins/Experimental/Animation/RelativeIKOp/Source/RelativeBodyAnimUtils/Private/RelativePropsAnimModifier.cpp b/Engine/Plugins/Experimental/Animation/RelativeIKOp/Source/RelativeBodyAnimUtils/Private/RelativePropsAnimModifier.cpp
--- a/Engine/Plugins/Experimental/Animation/RelativeIKOp/Source/RelativeBodyAnimUtils/Private/RelativePropsAnimModifier.cpp
+++ b/Engine/Plugins/Experimental/Animation/RelativeIKOp/Source/RelativeBodyAnimUtils/Private/RelativePropsAnimModifier.cpp
@@ -306,9 +306,13 @@ void URelativePropsAnimModifier::OnApply_Implementation(UAnimSequence* InAnimati
 
 							FVector3f Loc1(NearestPoint);
 							FVector3f Loc2 = FRelativeBodyAnimUtils::CalcVertLocationInUnitBody(ContactPoint2, BoneName2, RefSkeleton, CacheToLocals, CachedBodySourceData.SourceRetargetGlobalPose, PhysicsAsset);
+							FTransform3f PropOffsetTfm(BakedAnimNotifyInfo->OffsetTransformsForBones[BoneName1]);
+							FTransform3f SourceBodyRot(FTransform(BakedAnimNotifyInfo->OffsetTransformsForBones[BoneName2].GetRotation()));
 
-							BakedAnimNotifyInfo->PropsPairsLocalReference[Offset + 0] = Loc1;
-							BakedAnimNotifyInfo->PropsPairsLocalReference[Offset + 1] = Loc2;
+							// BakedAnimNotifyInfo->PropsPairsLocalReference[Offset + 0] = Loc1;
+							// BakedAnimNotifyInfo->PropsPairsLocalReference[Offset + 1] = Loc2;
+							BakedAnimNotifyInfo->PropsPairsLocalReference[Offset + 0] = PropOffsetTfm.TransformPosition(Loc1);
+							BakedAnimNotifyInfo->PropsPairsLocalReference[Offset + 1] = SourceBodyRot.TransformPosition(Loc2);
 							//UE_LOG(LogAnimation, Warning, TEXT("NumVertsIDX %d: (%f,%f,%f)."), NearestQueryVertexIndex, Loc1.X, Loc1.Y, Loc1.Z);
 						}
 						Count++;
diff --git a/Engine/Plugins/Experimental/Animation/RelativeIKOp/Source/RelativeIKOp/Private/RelativeIKOp.cpp b/Engine/Plugins/Experimental/Animation/RelativeIKOp/Source/RelativeIKOp/Private/RelativeIKOp.cpp
--- a/Engine/Plugins/Experimental/Animation/RelativeIKOp/Source/RelativeIKOp/Private/RelativeIKOp.cpp
+++ b/Engine/Plugins/Experimental/Animation/RelativeIKOp/Source/RelativeIKOp/Private/RelativeIKOp.cpp
@@ -22,6 +22,29 @@
 
 #define LOCTEXT_NAMESPACE "RelativeIKOp"
 
+template<int P>
+double RationalSmoothStepdown(double Distance, double DistThreshold, double DistHalfPointLambda)
+{
+	constexpr double kEps = UE_DOUBLE_SMALL_NUMBER;
+	
+	double x = FMath::Clamp(Distance / DistThreshold, 0., 1.);
+	double a = FMath::Clamp(DistHalfPointLambda, kEps, 1.-kEps);
+	const double ax = (1.0-a)*x;
+	const double iax = a*(1.0-x);
+	
+	double r1 = ax;
+	double r2 = iax;
+	if constexpr (P > 1) 
+	{
+		for (int i = 0; i < P - 1; ++i)
+		{
+			r1 *= ax;
+			r2 *= iax;
+		}
+	}
+	return (r2) / (r1 + r2);
+}
+
 const UClass* FIKRetargetRelativeIKOpSettings::GetControllerType() const
 {
 	return UIKRetargetRelativeIKController::StaticClass();
@@ -190,7 +213,11 @@ void FIKRetargetRelativeIKOp::Run(
 	}
 
 #if WITH_EDITOR
-	LocalDebugInfo.SourceDomainBones = CacheSourceDomainBones;
+	LocalDebugInfo.TargetDomainBones.SetNum(CacheSourceDomainBones.Num());
+	for (int i=0; i<CacheSourceDomainBones.Num(); i++)
+	{
+		LocalDebugInfo.TargetDomainBones[i] = ApplyBodyMap(CacheSourceDomainBones[i]);
+	}
 	UpdateDebugInfo(SourceScale, InSourceGlobalPose, OutTargetGlobalPose, LocalDebugInfo);
 #endif
 }
@@ -277,20 +304,12 @@ void FIKRetargetRelativeIKOp::ComputeFramePairSpaces(
 
 		int32 SourcePrimaryPoseIdx = CacheSourceSkelIndices[SourcePrimaryBone];
 		int32 SourceOtherPoseIdx = CacheSourceSkelIndices[SourceOtherBone];
-
-		// ---- TESTING ONLY - Remove these once notify baking is rotation invariant
-		FTransform SourcePrimaryBodyRot(CachedNotifyInfo->OffsetTransformsForBones[SourcePrimaryBone].GetRotation());
-		FTransform SourceOtherBodyRot(CachedNotifyInfo->OffsetTransformsForBones[SourceOtherBone].GetRotation());
 		
 		// Compute source transforms and source-space vertex distance
 		// const FVector3f& LocalPrimaryVert = CachedNotifyInfo->BodyPairsLocalReference[FramePairIdx + 2*PairIdx];
 		// const FVector3f& LocalOtherVert = CachedNotifyInfo->BodyPairsLocalReference[FramePairIdx + 2*PairIdx + 1];
-		const FVector3f& LocalPrimaryVert = FramePointView[2*PairIdx];
-		const FVector3f& LocalOtherVert = FramePointView[2*PairIdx+1];
-
-		// TODO: Clean this up if we bake in rot invariant space
-		FVector SourceBoneBakePrimaryV = (SourcePrimaryBodyRot).TransformPosition(FVector(LocalPrimaryVert)); 
-		FVector SourceBoneBakeOtherV = (SourceOtherBodyRot).TransformPosition(FVector(LocalOtherVert));
+		FVector SourceBoneBakePrimaryV(FramePointView[2*PairIdx]);
+		FVector SourceBoneBakeOtherV(FramePointView[2*PairIdx+1]);
 
 		FBodyTransform SourcePrimaryTfm;
 		FBodyTransform SourceOtherTfm;
@@ -308,7 +327,7 @@ void FIKRetargetRelativeIKOp::ComputeFramePairSpaces(
 		
 		double SourceVDist = FVector::Distance(RescaleSourcePrimaryV, RescaleSourceOtherV);
 		double LinearWeight = FMath::Clamp(1.0 - (SourceVDist / Settings.DistanceThreshold), 0.0, 1.0);
-		double Weight = GetDistanceWeight(SourceVDist, Settings.DistanceThreshold, Settings.DistanceFade);
+		double Weight = GetDistanceWeight(Settings.DistanceWeightMode, SourceVDist, Settings.DistanceThreshold, Settings.DistHalfPointLambda);
 
 		// Compute Source pair verts as represented in alternate pair body primary(_P) or other(_O) local bake space
 		FVector SourcePrimaryV_O = InverseBodyTransform(SourceOtherTfm, SourcePrimaryCompV);
@@ -449,23 +468,14 @@ void FIKRetargetRelativeIKOp::ComputeFramePropSpaces(
 		
 		int32 SourcePropPoseIdx = CacheSourceSkelIndices[SourcePropBone];
 		int32 SourceBodyPoseIdx = CacheSourceSkelIndices[SourceBodyBone];
-
-		// ---- TESTING ONLY - Remove these once notify baking is rotation invariant
-		const FTransform PropOffsetTfm = CachedPropNotifyInfo->OffsetTransformsForBones[SourcePropBone];
-		FTransform SourceBodyRot(CachedPropNotifyInfo->OffsetTransformsForBones[SourceBodyBone].GetRotation());
-		FTransform TargetBodyRot = SourceBodyRot;
 		
 		// Compute source transforms and source-space vertex distance
-		// const FVector3f& LocalPrimaryVert = CachedNotifyInfo->BodyPairsLocalReference[FramePairIdx + 2*PairIdx];
-		// const FVector3f& LocalOtherVert = CachedNotifyInfo->BodyPairsLocalReference[FramePairIdx + 2*PairIdx + 1];
-		const FVector3f& LocalPropVert = FramePointView[2*PairIdx];
-		const FVector3f& LocalBodyVert = FramePointView[2*PairIdx+1];
+		// const FVector3f& LocalPropVert = FramePointView[2*PairIdx];
+		// const FVector3f& LocalBodyVert = FramePointView[2*PairIdx+1];
+		FVector SourceBoneBakePropV(FramePointView[2*PairIdx]);
+		FVector SourceBoneBakeBodyV(FramePointView[2*PairIdx+1]);
 
 		bool bBodyMovable = CacheSourceDomainBones.Contains(SourceBodyBone);
-
-		// TODO: Clean this up if we bake in rot invariant space
-		FVector SourceBoneBakePropV = (PropOffsetTfm).TransformPosition(FVector(LocalPropVert));
-		FVector SourceBoneBakeBodyV = (SourceBodyRot).TransformPosition(FVector(LocalBodyVert));
 		
 		FTransform SourcePropTfm = SourceScaleTfm * SourcePose[SourcePropPoseIdx];
 		FBodyTransform SourceBodyTfm;
@@ -482,7 +492,7 @@ void FIKRetargetRelativeIKOp::ComputeFramePropSpaces(
 
 		double SourceVDist = FVector::Distance(RescaleSourcePropV, RescaleSourceBodyV);
 		double LinearWeight = FMath::Clamp(1.0 - (SourceVDist / Settings.DistanceThreshold), 0.0, 1.0);
-		double Weight = GetDistanceWeight(SourceVDist, Settings.DistanceThreshold, Settings.DistanceFade);
+		double Weight = GetDistanceWeight(Settings.DistanceWeightMode, SourceVDist, Settings.DistanceThreshold, Settings.DistHalfPointLambda);
 
 		// Compute Source pair verts as represented in alternate pair body primary(_P) or other(_O) local bake space
 		FVector SourcePropV_B = InverseBodyTransform(SourceBodyTfm, SourcePropCompV);
@@ -1022,10 +1032,26 @@ FKShapeElem* FIKRetargetRelativeIKOp::FindBodyShape(const UPhysicsAsset* PhysAss
 	return PhysAsset->SkeletalBodySetups[BodyIdx]->AggGeom.GetElement(0);
 }
 
-double FIKRetargetRelativeIKOp::GetDistanceWeight(double Distance, double DistThreshold, double ScalarFade)
+double FIKRetargetRelativeIKOp::GetDistanceWeight(EDistanceWeightMode DistanceWeightMode, double Distance, double DistThreshold, double DistHalfPointLambda)
 {
-	double Barrier = -(DistThreshold-Distance)*FMath::Max(DistThreshold-Distance, 0.0f) * FMath::Log2(Distance / DistThreshold);
-	return Barrier / (Barrier + ScalarFade); 
+	switch (DistanceWeightMode)
+	{
+	case EDistanceWeightMode::Linear:
+		return RationalSmoothStepdown<1>(Distance, DistThreshold, DistHalfPointLambda);
+	case EDistanceWeightMode::Quadratic:
+		return RationalSmoothStepdown<2>(Distance, DistThreshold, DistHalfPointLambda);
+	case EDistanceWeightMode::Cubic:
+		return RationalSmoothStepdown<3>(Distance, DistThreshold, DistHalfPointLambda);
+	case EDistanceWeightMode::Quartic:
+		return RationalSmoothStepdown<4>(Distance, DistThreshold, DistHalfPointLambda);
+	default: 
+		if (Distance > DistThreshold)
+		{
+			return 0.;
+		}
+		return 1.;
+	}
+	
 }
 
 void FIKRetargetRelativeIKOp::AnimGraphPreUpdateMainThread(USkeletalMeshComponent& SourceMeshComponent, USkeletalMeshComponent& TargetMeshComponent)
@@ -1316,10 +1342,17 @@ void FIKRetargetRelativeIKOp::UpdateCacheNotifyInfo(URelativeBodyBakeAnimNotify*
 	{
 		return;
 	}
-	
-	CachedNotifyInfo = NotifyInfo;
-	CacheMeshPairOffset = FrameBoneOffsets.Num() / 2;
 
+	if (!CachedNotifyInfo)
+	{
+		CachedNotifyInfo = NotifyInfo;
+	}
+	else
+	{
+		return;
+	}
+	CacheMeshPairOffset = FrameBoneOffsets.Num() / 2;
+	
 	int32 NumBodies = CachedNotifyInfo->BodyPairs.Num();
 	FrameBoneOffsets.AddZeroed(NumBodies);
 	FrameBonePositions.AddZeroed(NumBodies);
@@ -1379,7 +1412,14 @@ void FIKRetargetRelativeIKOp::UpdateCachePropNotifyInfo(URelativePropsBakeAnimNo
 		return;
 	}
 	
-	CachedPropNotifyInfo = NotifyInfo;
+	if (!CachedPropNotifyInfo)
+	{
+		CachedPropNotifyInfo = NotifyInfo;
+	}
+	else
+	{
+		return;
+	}
 	CachePropPairOffset = FrameBoneOffsets.Num() / 2;
 
 	int32 NumBones = CachedPropNotifyInfo->PropsPairs.Num();
@@ -1615,7 +1655,8 @@ void FIKRetargetRelativeIKOp::DebugDrawBodies(FPrimitiveDrawInterface* InPDI, co
 {
 	for (const FDebugDrawBodyInfo& BodyInfo : PhysBodies)
 	{
-		FKShapeElem* ShapeElem = FindBodyShape(PhysAsset, BodyInfo.BodyName);
+		FName TargetBodyName = TargetPAForSource? ApplyBodyMap(BodyInfo.BodyName) : BodyInfo.BodyName;
+		FKShapeElem* ShapeElem = FindBodyShape(PhysAsset, TargetBodyName);
 		if (!ShapeElem)
 		{
 			continue;
@@ -1633,7 +1674,7 @@ void FIKRetargetRelativeIKOp::DebugDrawBodies(FPrimitiveDrawInterface* InPDI, co
 		}
 
 		// Draw physics body
-		if (DebugDrawInfo.SourceDomainBones.Contains(BodyInfo.BodyName))
+		if (DebugDrawInfo.TargetDomainBones.Contains(TargetBodyName))
 		{
 			DebugDrawPhysBody(InPDI, CompTfm, Scale * BodyInfo.Scale, ShapeElem, FLinearColor::Black);
 		}
@@ -1722,7 +1763,7 @@ void FIKRetargetRelativeIKOp::UpdateDebugInfo(
 	
 	FScopeLock ScopeLock(&DebugDataMutex);
 	
-	DebugDrawInfo.SourceDomainBones = LocalDebugInfo.SourceDomainBones;
+	DebugDrawInfo.TargetDomainBones = LocalDebugInfo.TargetDomainBones;
 
 	DebugDrawInfo.TargetGoals = LocalDebugInfo.TargetGoals;
 	DebugDrawInfo.PairRetargetInfo = LocalDebugInfo.PairRetargetInfo;
diff --git a/Engine/Plugins/Experimental/Animation/RelativeIKOp/Source/RelativeIKOp/Public/RelativeIKOp.h b/Engine/Plugins/Experimental/Animation/RelativeIKOp/Source/RelativeIKOp/Public/RelativeIKOp.h
--- a/Engine/Plugins/Experimental/Animation/RelativeIKOp/Source/RelativeIKOp/Public/RelativeIKOp.h
+++ b/Engine/Plugins/Experimental/Animation/RelativeIKOp/Source/RelativeIKOp/Public/RelativeIKOp.h
@@ -19,6 +19,17 @@ class UPhysicsAsset;
 class URelativeBodyBakeAnimNotify;
 class URelativePropsBakeAnimNotify;
 
+UENUM(BlueprintType)
+enum class EDistanceWeightMode : uint8
+{
+	// Constant   UMETA(DisplayName = "Constant"),
+	Linear     UMETA(DisplayName = "Linear"),
+	Quadratic  UMETA(DisplayName = "Quadratic"),
+	Cubic      UMETA(DisplayName = "Cubic"),
+	Quartic    UMETA(DisplayName = "Quartic"),
+	// Quintic    UMETA(DisplayName = "Quintic")
+};
+
 // Debug info for individual body pair contributions
 struct FDebugRelativeTargetPairSpace
 {
@@ -84,7 +95,7 @@ struct FDebugRelativeIKDrawInfo
 	// Final goal points with contributions
 	TArray<FDebugDrawTargetGoal> TargetGoals;
 
-	TArray<FName> SourceDomainBones;
+	TArray<FName> TargetDomainBones;
 };
 
 USTRUCT(BlueprintType, meta = (DisplayName = "Relative IK Settings"))
@@ -112,9 +123,13 @@ struct FIKRetargetRelativeIKOpSettings : public FIKRetargetOpSettingsBase
 	UPROPERTY(EditAnywhere, BlueprintReadWrite, Category = "Parameters", meta=(ClampMin = "0.1"))
 	double DistanceThreshold = 50.00;
 
-	// Maximum distance for which body pair info is baked
-	UPROPERTY(EditAnywhere, BlueprintReadWrite, Category = "Parameters", meta=(ClampMin = "0"))
-	double DistanceFade = 200.00;
+	// DistHalfPointLambda*DistanceThreshold is the place where distance weight value be 0.5
+	UPROPERTY(EditAnywhere, BlueprintReadWrite, Category = "Parameters", meta=(ClampMin = "0.01", ClampMax = "0.99"))
+	double DistHalfPointLambda = 0.5;
+
+	// DistHalfPointLambda*DistanceThreshold is the place where distance weight value be 0.5
+	UPROPERTY(EditAnywhere, BlueprintReadWrite, Category = "Parameters")
+	EDistanceWeightMode DistanceWeightMode = EDistanceWeightMode::Quadratic;
 
 	// Bias feasibility distance weighting relative to distance threshold (+ increase target feasibility/- reduced target feasibility)
 	UPROPERTY(EditAnywhere, BlueprintReadWrite, Category = "Parameters")
@@ -322,7 +337,7 @@ private:
 	static FVector GetBodyOrientedScale(const UPhysicsAsset* PhysAsset, FName BoneName);
 	static FKShapeElem* FindBodyShape(const UPhysicsAsset* PhysAsset, FName BoneName);
 	
-	static double GetDistanceWeight(double Distance, double DistThreshold, double ScalarFade = 200.0);
+	static double GetDistanceWeight(EDistanceWeightMode DistanceWeightMode, double Distance, double DistThreshold, double DistHalfPointLambda = 0.5);
 
 #if WITH_EDITOR
 public:
