RelativeIKOp: Use matrix solve for multi-bone constraint positioning

diff --git a/Engine/Plugins/Experimental/Animation/RelativeIKOp/Source/RelativeIKOp/Private/RelativeIKOp.cpp b/Engine/Plugins/Experimental/Animation/RelativeIKOp/Source/RelativeIKOp/Private/RelativeIKOp.cpp
--- a/Engine/Plugins/Experimental/Animation/RelativeIKOp/Source/RelativeIKOp/Private/RelativeIKOp.cpp
+++ b/Engine/Plugins/Experimental/Animation/RelativeIKOp/Source/RelativeIKOp/Private/RelativeIKOp.cpp
@@ -12,6 +12,7 @@
 #include "PhysicsEngine/SkeletalBodySetup.h"
 #include "Retargeter/IKRetargetProcessor.h"
 #include "Retargeter/RetargetOps/RunIKRigOp.h"
+#include "Eigen/SVD"
 
 #if WITH_EDITOR
 #include "PrimitiveDrawingUtils.h"
@@ -150,8 +151,10 @@ void FIKRetargetRelativeIKOp::Run(
 
 	// Copy (pin) prop bones in target
 	// TODO: Pin in local space instead of comp space?
-	for (FName SourcePropBone : CachePinSourcePropBones)
+	CachePinTargetPropTransforms.SetNumZeroed(CachePinSourcePropBones.Num());
+	for (int32 PropIdx = 0; PropIdx < CachePinSourcePropBones.Num(); ++PropIdx)
 	{
+		FName SourcePropBone = CachePinSourcePropBones[PropIdx];
 		FName TargetPropBone = ApplyBodyMap(SourcePropBone);
 		if (!CacheSourceSkelIndices.Contains(SourcePropBone) || !CacheTargetSkelIndices.Contains(TargetPropBone))
 		{
@@ -161,6 +164,9 @@ void FIKRetargetRelativeIKOp::Run(
 		int32 SourceBoneIdx = CacheSourceSkelIndices[SourcePropBone];
 		int32 TargetBoneIdx = CacheTargetSkelIndices[TargetPropBone];
 		OutTargetGlobalPose[TargetBoneIdx] = InSourceGlobalPose[SourceBoneIdx];
+
+		// Copy comp-space pose to re-apply regardless
+		CachePinTargetPropTransforms[PropIdx] = OutTargetGlobalPose[TargetBoneIdx];
 	}
 
 	// TODO: Should we just use initial chain lengths?
@@ -188,28 +194,14 @@ void FIKRetargetRelativeIKOp::Run(
 	
 	LocalDebugInfo.TargetGoals.Reserve(CacheSourceBodyEffectVertIdx.Num());
 
-	CachePinTargetPropTransforms.SetNumZeroed(CachePinSourcePropBones.Num());
 	// Update bone goals to relative IK positions
-	for (int EffectIdx = 0; EffectIdx < CacheSourceBodyEffectVertIdx.Num(); ++EffectIdx)
+	if (Settings.bMultiBoneSolve)
 	{
-		TArray<int32>& EffectIndices = CacheSourceBodyEffectVertIdx[EffectIdx];
-		if (EffectIndices.IsEmpty())
-		{
-			continue;
-		}
-
-		FName SourceBoneName = CacheSourceEffectBones[EffectIdx];
-		FName TargetBoneName = ApplyBodyMap(SourceBoneName);
-		if (!CacheTargetSkelIndices.Contains(TargetBoneName))
-		{
-			continue;
-		}
-
-		int32 TargetBoneIdx = CacheTargetSkelIndices[TargetBoneName];
-		const FTransform& TargetBonePose = OutTargetGlobalPose[TargetBoneIdx];
-		UpdateTargetIKGoal(IKGoals, TargetBoneName, TargetBonePose, EffectIndices, LocalDebugInfo, LocalTargetPairSpaces);
-		
-		UpdatePropPinTransform(SourceBoneName, TargetBonePose, EffectIndices, LocalDebugInfo, LocalTargetPairSpaces);
+		ApplyMatrixBodyPairTargets(IKGoals, OutTargetGlobalPose, LocalDebugInfo, LocalTargetPairSpaces);
+	}
+	else
+	{
+		ApplyAverageBodyPairTargets(IKGoals, OutTargetGlobalPose, LocalDebugInfo, LocalTargetPairSpaces);
 	}
 
 #if WITH_EDITOR
@@ -586,30 +578,152 @@ void FIKRetargetRelativeIKOp::DropPropDomainBodyEffects(int32 PairOffset)
 	}
 }
 
-void FIKRetargetRelativeIKOp::UpdateTargetIKGoal(
+void FIKRetargetRelativeIKOp::ApplyAverageBodyPairTargets(
 	TArray<FIKRigGoal>& OutIKGoals,
-	FName TargetBoneName,
-	const FTransform& TargetBonePose,
-	const TArray<int32>& EffectIndices,
+	const TArray<FTransform>& TargetPose,
 	FDebugRelativeIKDrawInfo& DebugInfo,
 	TArray<FDebugRelativeTargetPairSpace>& DebugPairSpaces)
 {
-	FIKRigGoal* IKGoal = OutIKGoals.FindByPredicate([TargetBoneName](const FIKRigGoal& Other){return Other.BoneName == TargetBoneName;});
-	if (!IKGoal)
+	for (int32 EffectIdx = 0; EffectIdx < CacheSourceBodyEffectVertIdx.Num(); ++EffectIdx)
 	{
-		return;
+		TArray<int32>& EffectIndices = CacheSourceBodyEffectVertIdx[EffectIdx];
+		if (EffectIndices.IsEmpty())
+		{
+			continue;
+		}
+
+		FName SourceBoneName = CacheSourceEffectBones[EffectIdx];
+		FName TargetBoneName = ApplyBodyMap(SourceBoneName);
+		if (!CacheTargetSkelIndices.Contains(TargetBoneName))
+		{
+			continue;
+		}
+
+		int32 TargetBoneIdx = CacheTargetSkelIndices[TargetBoneName];
+		const FTransform& TargetBonePose = TargetPose[TargetBoneIdx];
+		FVector BoneLoc = TargetBonePose.GetLocation();
+
+		FVector TargetOffset;
+		double TotalWeight = ComputeTargetWeightedOffset(TargetOffset, EffectIndices, BoneLoc, DebugInfo, DebugPairSpaces);
+		if (FMath::IsNearlyZero(TotalWeight))
+		{
+			return;
+		}
+
+		FVector TargetLoc = BoneLoc + TargetOffset;
+		UpdateTargetIKGoal(OutIKGoals, TargetBoneName, TotalWeight, TargetLoc, TargetBonePose);
+		UpdatePropPinTransform(SourceBoneName, TotalWeight, TargetLoc, TargetBonePose);
 	}
+}
 
-	FVector TargetGoalLoc;
-	double TotalWeight = ComputeTargetWeightedLocation(TargetGoalLoc, EffectIndices, DebugInfo, DebugPairSpaces);
-	if (FMath::IsNearlyZero(TotalWeight))
+
+void FIKRetargetRelativeIKOp::ApplyMatrixBodyPairTargets(
+	TArray<FIKRigGoal>& OutIKGoals,
+	const TArray<FTransform>& TargetPose,
+	FDebugRelativeIKDrawInfo& DebugInfo,
+	TArray<FDebugRelativeTargetPairSpace>& DebugPairSpaces)
+{
+	TArray<int32> OutIndexMap;
+	TArray<FName> OutBoneNames;
+	TArray<FVector> OutBoneLocations;
+	
+	OutIndexMap.Reserve(CacheSourceBodyEffectVertIdx.Num());
+	OutBoneNames.Reserve(CacheSourceBodyEffectVertIdx.Num());
+	OutBoneLocations.Reserve(CacheSourceBodyEffectVertIdx.Num());
+	for (int EffectIdx = 0; EffectIdx < CacheSourceBodyEffectVertIdx.Num(); ++EffectIdx)
+	{
+		TArray<int32>& EffectIndices = CacheSourceBodyEffectVertIdx[EffectIdx];
+		if (EffectIndices.IsEmpty())
+		{
+			continue;
+		}
+
+		FName SourceBoneName = CacheSourceEffectBones[EffectIdx];
+		FName TargetBoneName = ApplyBodyMap(SourceBoneName);
+		if (!CacheTargetSkelIndices.Contains(TargetBoneName))
+		{
+			continue;
+		}
+
+		int32 TargetBoneIdx = CacheTargetSkelIndices[TargetBoneName];
+		const FTransform& TargetBonePose = TargetPose[TargetBoneIdx];
+		
+		OutIndexMap.Add(EffectIdx);
+		OutBoneNames.Add(SourceBoneName);
+		OutBoneLocations.Add(TargetBonePose.GetLocation());
+	}
+
+	Eigen::MatrixXd WeightMatrix(OutIndexMap.Num(), OutIndexMap.Num());
+	Eigen::MatrixXd OffsetMatrix(OutIndexMap.Num(), 3);
+	Eigen::VectorXd TotalWeights(OutIndexMap.Num());
+
+	WeightMatrix = Eigen::MatrixXd::Identity(OutIndexMap.Num(), OutIndexMap.Num());
+	OffsetMatrix = Eigen::MatrixXd::Zero(OutIndexMap.Num(), 3);
+	TotalWeights = Eigen::VectorXd::Zero(OutIndexMap.Num());
+	for (int32 RowIdx = 0; RowIdx < OutIndexMap.Num(); ++RowIdx)
+	{
+		int32 EffectIdx = OutIndexMap[RowIdx];
+		TArray<int32>& EffectIndices = CacheSourceBodyEffectVertIdx[EffectIdx];
+		if (EffectIndices.IsEmpty())
+		{
+			continue;
+		}
+		
+		double TotalWeight = ComputeTargetWeightedOffsetRow(OffsetMatrix, RowIdx, OutBoneNames, EffectIndices, OutBoneLocations[RowIdx], DebugInfo, DebugPairSpaces);
+		if (FMath::IsNearlyZero(TotalWeight))
+		{
+			continue;
+		}
+		
+		ComputeWeightMatrixRow(WeightMatrix, RowIdx, OutBoneNames, EffectIndices);
+		TotalWeights(RowIdx) = TotalWeight;
+		WeightMatrix(RowIdx, RowIdx) = TotalWeight;
+	}
+
+	Eigen::MatrixXd SolvedOffset = WeightMatrix.jacobiSvd(Eigen::ComputeThinU | Eigen::ComputeThinV).solve(OffsetMatrix);
+	// Eigen::MatrixXd SolvedOffset = WeightMatrix.ldlt().solve(OffsetMatrix);
+	for (int32 RowIdx = 0; RowIdx < OutIndexMap.Num(); ++RowIdx)
+	{
+		int32 EffectIdx = OutIndexMap[RowIdx];
+		if (FMath::IsNearlyZero(TotalWeights(RowIdx)))
+		{
+			continue;
+		}
+		
+		FName SourceBoneName = CacheSourceEffectBones[EffectIdx];
+		FName TargetBoneName = ApplyBodyMap(SourceBoneName);
+		if (!CacheTargetSkelIndices.Contains(TargetBoneName))
+		{
+			continue;
+		}
+		
+		FVector TargetOffset{SolvedOffset(RowIdx, 0), SolvedOffset(RowIdx, 1), SolvedOffset(RowIdx, 2)};
+
+		int32 TargetBoneIdx = CacheTargetSkelIndices[TargetBoneName];
+		const FTransform& TargetBonePose = TargetPose[TargetBoneIdx];
+
+		FVector TargetLoc = TargetBonePose.GetLocation() + TargetOffset;
+		UpdateTargetIKGoal(OutIKGoals, TargetBoneName, TotalWeights(RowIdx), TargetLoc, TargetBonePose);
+		UpdatePropPinTransform(SourceBoneName, TotalWeights(RowIdx), TargetLoc, TargetBonePose);
+	}
+}
+
+void FIKRetargetRelativeIKOp::UpdateTargetIKGoal(
+	TArray<FIKRigGoal>& OutIKGoals,
+	FName TargetBoneName,
+	double TargetWeight,
+	const FVector& TargetGoalLoc,
+	const FTransform& InTargetBonePose)
+{
+	FIKRigGoal* IKGoal = OutIKGoals.FindByPredicate([TargetBoneName](const FIKRigGoal& Other){return Other.BoneName == TargetBoneName;});
+	if (!IKGoal)
 	{
 		return;
 	}
 		
 	// TODO: Play with total weight functions to alpha goal
-	double GoalAlpha = FMath::Min(1.0, TotalWeight*FMath::Clamp(Settings.ContributionSumWeight, 0.0f, 1.0f));
-	FVector PrevGoal = GoalLocBlendCompSpace(IKGoal, TargetBonePose);
+	double GoalAlpha = FMath::Min(1.0, TargetWeight*FMath::Clamp(Settings.ContributionSumWeight, 0.0f, 1.0f));
+	FVector PrevGoal = GoalLocBlendCompSpace(IKGoal, InTargetBonePose);
 	FVector BlendGoal = FMath::Lerp(PrevGoal, TargetGoalLoc, GoalAlpha);
 
 	// TODO: Remove dry run for release
@@ -626,42 +740,32 @@ void FIKRetargetRelativeIKOp::UpdateTargetIKGoal(
 
 void FIKRetargetRelativeIKOp::UpdatePropPinTransform(
 	FName SourceBoneName,
-	const FTransform& TargetBonePose,
-	const TArray<int32>& EffectIndices,
-	FDebugRelativeIKDrawInfo& DebugInfo,
-	TArray<FDebugRelativeTargetPairSpace>& DebugPairSpaces)
+	double TargetWeight,
+	const FVector& TargetPropLoc,
+	const FTransform& InTargetBonePose)
 {
 	int32 PinPropIdx = CachePinSourcePropBones.Find(SourceBoneName);
 	if (PinPropIdx == INDEX_NONE)
 	{
 		return;
 	}
-
-	// Copy comp-space pose to re-apply regardless
-	CachePinTargetPropTransforms[PinPropIdx] = TargetBonePose;
-
-	FVector TargetPropBoneLoc;
-	double TotalWeight = ComputeTargetWeightedLocation(TargetPropBoneLoc, EffectIndices, DebugInfo, DebugPairSpaces);
-	if (FMath::IsNearlyZero(TotalWeight))
-	{
-		return;
-	}
 	
-	double PropBoneAlpha = FMath::Min(1.0, TotalWeight*FMath::Clamp(Settings.ContributionSumWeight, 0.0, 1.0));
+	double PropBoneAlpha = FMath::Min(1.0, TargetWeight*FMath::Clamp(Settings.ContributionSumWeight, 0.0, 1.0));
 	double BlendAlpha = PropBoneAlpha * (1.0-Settings.BlendPinsToSource);
-	FVector BlendBoneLoc = FMath::Lerp(TargetBonePose.GetLocation(), TargetPropBoneLoc, BlendAlpha);
+	FVector BlendBoneLoc = FMath::Lerp(InTargetBonePose.GetLocation(), TargetPropLoc, BlendAlpha);
 	
 	CachePinTargetPropTransforms[PinPropIdx].SetLocation(BlendBoneLoc);
 }
 
-double FIKRetargetRelativeIKOp::ComputeTargetWeightedLocation(
-	FVector& OutTargetLocation,
+double FIKRetargetRelativeIKOp::ComputeTargetWeightedOffset(
+	FVector& OutTargetOffset,
 	const TArray<int32>& EffectIndices,
+	const FVector& BoneLoc,
 	FDebugRelativeIKDrawInfo& DebugInfo,
 	TArray<FDebugRelativeTargetPairSpace>& DebugPairSpaces)
 {
 	double TotalWeight = 0;
-	OutTargetLocation = FVector::ZeroVector;
+	OutTargetOffset = FVector::ZeroVector;
 
 #if WITH_EDITOR
 	if (Settings.bDebugDraw)
@@ -670,22 +774,21 @@ double FIKRetargetRelativeIKOp::ComputeTargetWeightedLocation(
 		DebugInfo.TargetGoals.Last().Goal = FVector::ZeroVector;
 	}
 #endif //WITH_EDITOR
-
-	double MaxWeight = 0;
+	
 	for (int Index : EffectIndices)
 	{
 		if (FMath::IsNearlyZero(FrameBoneWeights[Index]))
 		{
 			continue;
 		}
-
-		FVector TargetLoc = FrameBoneOffsets[Index] + FrameBonePositions[Index];
-		OutTargetLocation += FrameBoneWeights[Index] * TargetLoc;
+		
+		OutTargetOffset += FrameBoneWeights[Index] * FrameBoneOffsets[Index];
 		TotalWeight += FrameBoneWeights[Index];
 
 #if WITH_EDITOR
 		if (Settings.bDebugDraw)
 		{
+			FVector TargetLoc = FrameBoneOffsets[Index] + BoneLoc;
 			DebugInfo.TargetGoals.Last().PairTargets.Add(TargetLoc);
 			DebugInfo.PairRetargetInfo.Emplace(DebugPairSpaces[Index]);
 		}
@@ -697,11 +800,99 @@ double FIKRetargetRelativeIKOp::ComputeTargetWeightedLocation(
 		return TotalWeight;
 	}
 		
-	OutTargetLocation /= TotalWeight;
+	OutTargetOffset /= TotalWeight;
 #if WITH_EDITOR
 	if (Settings.bDebugDraw)
 	{
-		DebugInfo.TargetGoals.Last().Goal = OutTargetLocation;
+		DebugInfo.TargetGoals.Last().Goal = OutTargetOffset + BoneLoc;
+	}
+#endif // WITH_EDITOR
+	
+	return TotalWeight;
+}
+
+
+
+void FIKRetargetRelativeIKOp::ComputeWeightMatrixRow(
+	Eigen::MatrixXd& WeightMatrix,
+	int32 RowIdx,
+	const TArray<FName>& OutBones,
+	const TArray<int32>& EffectIndices)
+{
+	for (int Index : EffectIndices)
+	{
+		if (FMath::IsNearlyZero(FrameBoneWeights[Index]))
+		{
+			continue;
+		}
+
+		FName EffectBone = CachePairEffectBoneNames[Index];
+		int32 OutWeightIdx = OutBones.Find(EffectBone);
+		if (OutWeightIdx == INDEX_NONE)
+		{
+			continue;
+		}
+
+		WeightMatrix(RowIdx, OutWeightIdx) = -FrameBoneWeights[Index];
+	}
+}
+
+double FIKRetargetRelativeIKOp::ComputeTargetWeightedOffsetRow(
+	Eigen::MatrixXd& OffsetMatrix,
+	int32 RowIdx,
+	const TArray<FName>& OutBones,
+	const TArray<int32>& EffectIndices,
+	const FVector& BoneLoc,
+	FDebugRelativeIKDrawInfo& DebugInfo,
+	TArray<FDebugRelativeTargetPairSpace>& DebugPairSpaces)
+{
+	double TotalWeight = 0;
+	FVector OffsetLoc = FVector::ZeroVector;
+
+#if WITH_EDITOR
+	FVector PresolveTarget = FVector::ZeroVector;
+	
+	if (Settings.bDebugDraw)
+	{
+		DebugInfo.TargetGoals.AddDefaulted();
+		DebugInfo.TargetGoals.Last().Goal = FVector::ZeroVector;
+	}
+#endif //WITH_EDITOR
+	
+	for (int Index : EffectIndices)
+	{
+		if (FMath::IsNearlyZero(FrameBoneWeights[Index]))
+		{
+			continue;
+		}
+
+		FVector TargetOffset = FrameBoneOffsets[Index];
+
+		OffsetLoc += FrameBoneWeights[Index] * TargetOffset;
+		TotalWeight += FrameBoneWeights[Index];
+
+#if WITH_EDITOR
+		if (Settings.bDebugDraw)
+		{
+			// TODO: Fixup target debug to show dynamic offset goals (post-solve update?)
+			FVector PresolveBoneTarget = TargetOffset + BoneLoc;
+			PresolveTarget += FrameBoneWeights[Index] * PresolveBoneTarget;
+				
+			DebugInfo.TargetGoals.Last().PairTargets.Add(PresolveBoneTarget);
+			DebugInfo.PairRetargetInfo.Emplace(DebugPairSpaces[Index]);
+		}
+#endif //WITH_EDITOR
+	}
+
+	// Each matrix column is a dimension (x,y,z)
+	OffsetMatrix(RowIdx,0) = OffsetLoc.X;
+	OffsetMatrix(RowIdx,1) = OffsetLoc.Y;
+	OffsetMatrix(RowIdx,2) = OffsetLoc.Z;
+	
+#if WITH_EDITOR
+	if (Settings.bDebugDraw)
+	{
+		DebugInfo.TargetGoals.Last().Goal = PresolveTarget / TotalWeight;
 	}
 #endif // WITH_EDITOR
 	
@@ -1330,6 +1521,7 @@ void FIKRetargetRelativeIKOp::ResetCacheNotifyInfo()
 	CacheSourceEffectBones.Reset();
 	CacheSourceDomainBones.Reset();
 	CacheBodyBones.Reset();
+	CachePairEffectBoneNames.Reset();
 	FrameBoneOffsets.Reset();
 	FrameBonePositions.Reset();
 	FrameBoneWeights.Reset();
@@ -1357,6 +1549,7 @@ void FIKRetargetRelativeIKOp::UpdateCacheNotifyInfo(URelativeBodyBakeAnimNotify*
 	FrameBoneOffsets.AddZeroed(NumBodies);
 	FrameBonePositions.AddZeroed(NumBodies);
 	FrameBoneWeights.AddZeroed(NumBodies);
+	CachePairEffectBoneNames.AddDefaulted(NumBodies);
 
 	MeshSampleRate = 30.0f;
 	if (CachedNotifyInfo->BodyPairsSampleTime.Num() >= 2)
@@ -1402,6 +1595,10 @@ void FIKRetargetRelativeIKOp::UpdateCacheNotifyInfo(URelativeBodyBakeAnimNotify*
 			// Primary bone can be affect if e.g. hands clapping
 			CacheSourceBodyEffectVertIdx[PrimaryIdx].Add(2*OutPairIdx);
 		}
+
+		// TODO: Normalize pair edge info/frame info
+		CachePairEffectBoneNames[2*OutPairIdx] = OtherBone;
+		CachePairEffectBoneNames[2*OutPairIdx + 1] = PrimaryBone;
 	}
 }
 
@@ -1426,6 +1623,7 @@ void FIKRetargetRelativeIKOp::UpdateCachePropNotifyInfo(URelativePropsBakeAnimNo
 	FrameBoneOffsets.AddZeroed(NumBones);
 	FrameBonePositions.AddZeroed(NumBones);
 	FrameBoneWeights.AddZeroed(NumBones);
+	CachePairEffectBoneNames.AddDefaulted(NumBones);
 
 	PropSampleRate = 30.0f;
 	if (CachedPropNotifyInfo->PropsPairsSampleTime.Num() >= 2)
@@ -1463,6 +1661,10 @@ void FIKRetargetRelativeIKOp::UpdateCachePropNotifyInfo(URelativePropsBakeAnimNo
 		int32 PropIdx = CacheSourceEffectBones.Find(PropBone);
 		int32 BodyIdx = CacheSourceEffectBones.Find(BodyBone);
 
+		// TODO: Normalize pair edge info for all bake types
+		CachePairEffectBoneNames[2*OutPairIdx] = BodyBone;
+		CachePairEffectBoneNames[2*OutPairIdx + 1] = PropBone;
+
 		if (PropIdx == INDEX_NONE || BodyIdx == INDEX_NONE)
 		{
 			continue;
diff --git a/Engine/Plugins/Experimental/Animation/RelativeIKOp/Source/RelativeIKOp/Public/RelativeIKOp.h b/Engine/Plugins/Experimental/Animation/RelativeIKOp/Source/RelativeIKOp/Public/RelativeIKOp.h
--- a/Engine/Plugins/Experimental/Animation/RelativeIKOp/Source/RelativeIKOp/Public/RelativeIKOp.h
+++ b/Engine/Plugins/Experimental/Animation/RelativeIKOp/Source/RelativeIKOp/Public/RelativeIKOp.h
@@ -2,6 +2,8 @@
 #pragma once
 
 #include "Retargeter/IKRetargetOps.h"
+#include "Eigen/Dense"
+
 #include "RelativeIKOp.generated.h"
 
 #define UE_API RELATIVEIKOP_API
@@ -191,6 +193,10 @@ struct FIKRetargetRelativeIKOpSettings : public FIKRetargetOpSettingsBase
 	// UPROPERTY(BlueprintReadWrite, Category = "Test Features")
 	// bool bTestFeasibilityWeight = true;
 
+	// Solve for output positions using multiple interacting contact bodies
+	UPROPERTY(EditAnywhere, BlueprintReadWrite, Category = "Test Features")
+	bool bMultiBoneSolve = true;
+	
 	// Distance blend body pairs towards each contact position if both can move
 	UPROPERTY(EditAnywhere, BlueprintReadWrite, Category = "Test Features")
 	bool bBidirectionalDistanceContactBlend = true;
@@ -279,24 +285,50 @@ private:
 		FDebugRelativeIKDrawInfo& DebugInfo,
 		TArray<FDebugRelativeTargetPairSpace>& DebugPairSpaces);
 
-	void UpdateTargetIKGoal(
+	void ApplyAverageBodyPairTargets(
 		TArray<FIKRigGoal>& OutIKGoals,
-		FName TargetBoneName,
-		const FTransform& TargetBonePose,
-		const TArray<int32>& EffectIndices,
+		const TArray<FTransform>& TargetPose,
+		FDebugRelativeIKDrawInfo& DebugInfo,
+		TArray<FDebugRelativeTargetPairSpace>& DebugPairSpaces);
+	
+	void ApplyMatrixBodyPairTargets(
+		TArray<FIKRigGoal>& OutIKGoals,
+		const TArray<FTransform>& TargetPose,
 		FDebugRelativeIKDrawInfo& DebugInfo,
 		TArray<FDebugRelativeTargetPairSpace>& DebugPairSpaces);
 
+	void UpdateTargetIKGoal(
+		TArray<FIKRigGoal>& OutIKGoals,
+		FName TargetBoneName,
+		double TargetWeight,
+		const FVector& TargetGoalLoc,
+		const FTransform& InTargetBonePose);
+
 	void UpdatePropPinTransform(
 		FName SourceBoneName,
-		const FTransform& TargetBonePose,
+		double TargetWeight,
+		const FVector& TargetPropLoc,
+		const FTransform& InTargetBonePose);
+
+	double ComputeTargetWeightedOffset(
+		FVector& OutTargetOffset,
 		const TArray<int32>& EffectIndices,
+		const FVector& BoneLoc,
 		FDebugRelativeIKDrawInfo& DebugInfo,
 		TArray<FDebugRelativeTargetPairSpace>& DebugPairSpaces);
 
-	double ComputeTargetWeightedLocation(
-		FVector& OutTargetLocation,
+	void ComputeWeightMatrixRow(
+		Eigen::MatrixXd& WeightMatrix,
+		int32 RowIdx,
+		const TArray<FName>& OutBones,
+		const TArray<int32>& EffectIndices);
+
+	double ComputeTargetWeightedOffsetRow(
+		Eigen::MatrixXd& OffsetMatrix,
+		int32 RowIdx,
+		const TArray<FName>& OutBones,
 		const TArray<int32>& EffectIndices,
+		const FVector& BoneLoc,
 		FDebugRelativeIKDrawInfo& DebugInfo,
 		TArray<FDebugRelativeTargetPairSpace>& DebugPairSpaces);
 
@@ -395,6 +427,7 @@ private:
 	TArray<FName> CacheSourceEffectBones;
 	TArray<FName> CacheSourceDomainBones;
 	TSet<FName> CacheBodyBones;
+	TArray<FName> CachePairEffectBoneNames;
 	TArray<double> FrameBoneWeights;
 	TArray<FVector> FrameBonePositions;
 	TArray<FVector> FrameBoneOffsets;
diff --git a/Engine/Plugins/Experimental/Animation/RelativeIKOp/Source/RelativeIKOp/RelativeIKOp.Build.cs b/Engine/Plugins/Experimental/Animation/RelativeIKOp/Source/RelativeIKOp/RelativeIKOp.Build.cs
--- a/Engine/Plugins/Experimental/Animation/RelativeIKOp/Source/RelativeIKOp/RelativeIKOp.Build.cs
+++ b/Engine/Plugins/Experimental/Animation/RelativeIKOp/Source/RelativeIKOp/RelativeIKOp.Build.cs
@@ -48,7 +48,8 @@ namespace UnrealBuildTool.Rules
 				"SkeletalMeshDescription",
 				"Renderer",
 				"IKRig",
-				"RelativeBodyAnimInfo"
+				"RelativeBodyAnimInfo",
+				"Eigen"
 				}
 				);
 
