RelativeIKOp: Add regularization to multibone solve.

#Synced-CL 48120296

diff --git a/Engine/Plugins/Experimental/Animation/RelativeIKOp/Source/RelativeIKOp/Private/RelativeIKOp.cpp b/Engine/Plugins/Experimental/Animation/RelativeIKOp/Source/RelativeIKOp/Private/RelativeIKOp.cpp
--- a/Engine/Plugins/Experimental/Animation/RelativeIKOp/Source/RelativeIKOp/Private/RelativeIKOp.cpp
+++ b/Engine/Plugins/Experimental/Animation/RelativeIKOp/Source/RelativeIKOp/Private/RelativeIKOp.cpp
@@ -677,7 +677,7 @@ void FIKRetargetRelativeIKOp::ApplyMatrixBodyPairTargets(
 		
 		ComputeWeightMatrixRow(WeightMatrix, RowIdx, OutBoneNames, EffectIndices);
 		TotalWeights(RowIdx) = TotalWeight;
-		WeightMatrix(RowIdx, RowIdx) = TotalWeight;
+		WeightMatrix(RowIdx, RowIdx) = TotalWeight + Settings.MultiBoneRegularization;
 	}
 
 	Eigen::MatrixXd SolvedOffset = WeightMatrix.jacobiSvd(Eigen::ComputeThinU | Eigen::ComputeThinV).solve(OffsetMatrix);
diff --git a/Engine/Plugins/Experimental/Animation/RelativeIKOp/Source/RelativeIKOp/Public/RelativeIKOp.h b/Engine/Plugins/Experimental/Animation/RelativeIKOp/Source/RelativeIKOp/Public/RelativeIKOp.h
--- a/Engine/Plugins/Experimental/Animation/RelativeIKOp/Source/RelativeIKOp/Public/RelativeIKOp.h
+++ b/Engine/Plugins/Experimental/Animation/RelativeIKOp/Source/RelativeIKOp/Public/RelativeIKOp.h
@@ -197,6 +197,10 @@ struct FIKRetargetRelativeIKOpSettings : public FIKRetargetOpSettingsBase
 	UPROPERTY(EditAnywhere, BlueprintReadWrite, Category = "Test Features")
 	bool bMultiBoneSolve = true;
 	
+	// Regularization additive constant for stabilizing multibone solve
+	UPROPERTY(EditAnywhere, BlueprintReadWrite, Category = "Test Features", meta=(ClampMin = "0", ClampMax = "2"))
+	double MultiBoneRegularization = 0.5;
+	
 	// Distance blend body pairs towards each contact position if both can move
 	UPROPERTY(EditAnywhere, BlueprintReadWrite, Category = "Test Features")
 	bool bBidirectionalDistanceContactBlend = true;
