RelativeIKOp: Correctly identify and removed domain bodies from cache effect target setup. And fix copy-paste bug. Should fix several  instabilities in prop interactions.

#Synced-CL 48212065

diff --git a/Engine/Plugins/Experimental/Animation/RelativeIKOp/Source/RelativeIKOp/Private/RelativeIKOp.cpp b/Engine/Plugins/Experimental/Animation/RelativeIKOp/Source/RelativeIKOp/Private/RelativeIKOp.cpp
--- a/Engine/Plugins/Experimental/Animation/RelativeIKOp/Source/RelativeIKOp/Private/RelativeIKOp.cpp
+++ b/Engine/Plugins/Experimental/Animation/RelativeIKOp/Source/RelativeIKOp/Private/RelativeIKOp.cpp
@@ -194,6 +194,8 @@ void FIKRetargetRelativeIKOp::Run(
 	
 	LocalDebugInfo.TargetGoals.Reserve(CacheSourceBodyEffectVertIdx.Num());
 
+	// TODO: This is a hack to handle the fact that domain bodies aren't differentiated in prop bake data
+	DropPropDomainBodyEffects();
 	// Update bone goals to relative IK positions
 	if (Settings.bMultiBoneSolve)
 	{
@@ -565,7 +567,7 @@ void FIKRetargetRelativeIKOp::ComputeFramePropSpaces(
 	}
 }
 
-void FIKRetargetRelativeIKOp::DropPropDomainBodyEffects(int32 PairOffset)
+void FIKRetargetRelativeIKOp::DropPropDomainBodyEffects()
 {
 	// TODO: This should be handled by baking necessary info in prop/body notifies
 	for (int32 Index=0; Index < CacheSourceEffectBones.Num(); ++Index)
@@ -607,7 +609,7 @@ void FIKRetargetRelativeIKOp::ApplyAverageBodyPairTargets(
 		double TotalWeight = ComputeTargetWeightedOffset(TargetOffset, EffectIndices, BoneLoc, DebugInfo, DebugPairSpaces);
 		if (FMath::IsNearlyZero(TotalWeight))
 		{
-			return;
+			continue;
 		}
 
 		FVector TargetLoc = BoneLoc + TargetOffset;
@@ -1273,7 +1275,6 @@ void FIKRetargetRelativeIKOp::AnimGraphPreUpdateMainThread(USkeletalMeshComponen
 
 		ResetCacheNotifyInfo();
 		SetupRelativeIKNotifyInfoMontage();
-		PostCachedAllNotifyInfo();
 		return;
 	}
 
@@ -1490,19 +1491,6 @@ void FIKRetargetRelativeIKOp::UpdateRelativeIKNotifyInfoAnimSeq(UAnimInstance* S
 			UpdateCachePropNotifyInfo(PropInfo);
 		}
 	}
-
-	PostCachedAllNotifyInfo();
-}
-
-void FIKRetargetRelativeIKOp::PostCachedAllNotifyInfo()
-{
-	for (int EffectIdx = 0; EffectIdx < CacheSourceEffectBones.Num(); ++EffectIdx)
-	{
-		if (CacheSourceBodyEffectVertIdx[EffectIdx].IsEmpty())
-		{
-			CacheSourceDomainBones.Add(CacheSourceEffectBones[EffectIdx]);
-		}
-	}
 }
 
 void FIKRetargetRelativeIKOp::ResetCacheNotifyInfo()
@@ -1595,6 +1583,11 @@ void FIKRetargetRelativeIKOp::UpdateCacheNotifyInfo(URelativeBodyBakeAnimNotify*
 			// Primary bone can be affect if e.g. hands clapping
 			CacheSourceBodyEffectVertIdx[PrimaryIdx].Add(2*OutPairIdx);
 		}
+		else
+		{
+			// TODO: Bake necessary info in props to avoid having to store this data
+			CacheSourceDomainBones.Add(PrimaryBone);
+		}
 
 		// TODO: Normalize pair edge info/frame info
 		CachePairEffectBoneNames[2*OutPairIdx] = OtherBone;
diff --git a/Engine/Plugins/Experimental/Animation/RelativeIKOp/Source/RelativeIKOp/Public/RelativeIKOp.h b/Engine/Plugins/Experimental/Animation/RelativeIKOp/Source/RelativeIKOp/Public/RelativeIKOp.h
--- a/Engine/Plugins/Experimental/Animation/RelativeIKOp/Source/RelativeIKOp/Public/RelativeIKOp.h
+++ b/Engine/Plugins/Experimental/Animation/RelativeIKOp/Source/RelativeIKOp/Public/RelativeIKOp.h
@@ -198,8 +198,8 @@ struct FIKRetargetRelativeIKOpSettings : public FIKRetargetOpSettingsBase
 	bool bMultiBoneSolve = true;
 	
 	// Regularization additive constant for stabilizing multibone solve
-	UPROPERTY(EditAnywhere, BlueprintReadWrite, Category = "Test Features", meta=(ClampMin = "0", ClampMax = "2"))
-	double MultiBoneRegularization = 0.5;
+	UPROPERTY(EditAnywhere, BlueprintReadWrite, Category = "Test Features", meta=(ClampMin = "0", ClampMax = "1"))
+	double MultiBoneRegularization = 0.0;
 	
 	// Distance blend body pairs towards each contact position if both can move
 	UPROPERTY(EditAnywhere, BlueprintReadWrite, Category = "Test Features")
@@ -269,7 +269,7 @@ private:
 	void UpdateCacheChainInfo(const TArray<FTransform>& TargetPose, FName BoneName);
 	void UpdateCachePropNotifyInfo(URelativePropsBakeAnimNotify* NotifyInfo);
 
-	void DropPropDomainBodyEffects(int32 PairOffset);
+	void DropPropDomainBodyEffects();
 	
 	void ComputeFramePairSpaces(
 		double SourceScale,
