RelativeIKOp: Expose pin bone options similar to pin bone op (WIP).

#Synced-CL 48282260

diff --git a/Engine/Plugins/Experimental/Animation/RelativeIKOp/Source/RelativeIKOp/Private/RelativeIKOp.cpp b/Engine/Plugins/Experimental/Animation/RelativeIKOp/Source/RelativeIKOp/Private/RelativeIKOp.cpp
--- a/Engine/Plugins/Experimental/Animation/RelativeIKOp/Source/RelativeIKOp/Private/RelativeIKOp.cpp
+++ b/Engine/Plugins/Experimental/Animation/RelativeIKOp/Source/RelativeIKOp/Private/RelativeIKOp.cpp
@@ -96,6 +96,9 @@ bool FIKRetargetRelativeIKOp::Initialize(
 
 	UpdateCacheBoneChains(InProcessor, ParentRigOp);
 	UpdateCacheSkelInfo(InSourceSkeleton, InTargetSkeleton);
+	
+	// Add all pin bones to the source/target skel lookup tables
+	UpdatePinData(InProcessor);
 
 	// Force update cache info if new anim montage found
 	MontageInstance = nullptr;
@@ -149,25 +152,8 @@ void FIKRetargetRelativeIKOp::Run(
 	LocalTargetPairSpaces.SetNumZeroed(TotalBodies);
 #endif //WITH_EDITOR
 
-	// Copy (pin) prop bones in target
-	// TODO: Pin in local space instead of comp space?
-	CachePinTargetPropTransforms.SetNumZeroed(CachePinSourcePropBones.Num());
-	for (int32 PropIdx = 0; PropIdx < CachePinSourcePropBones.Num(); ++PropIdx)
-	{
-		FName SourcePropBone = CachePinSourcePropBones[PropIdx];
-		FName TargetPropBone = ApplyBodyMap(SourcePropBone);
-		if (!CacheSourceSkelIndices.Contains(SourcePropBone) || !CacheTargetSkelIndices.Contains(TargetPropBone))
-		{
-			continue;
-		}
-		
-		int32 SourceBoneIdx = CacheSourceSkelIndices[SourcePropBone];
-		int32 TargetBoneIdx = CacheTargetSkelIndices[TargetPropBone];
-		OutTargetGlobalPose[TargetBoneIdx] = InSourceGlobalPose[SourceBoneIdx];
-
-		// Copy comp-space pose to re-apply regardless
-		CachePinTargetPropTransforms[PropIdx] = OutTargetGlobalPose[TargetBoneIdx];
-	}
+	// Copy/cache (pin) prop bones in target using pin bone settings
+	CachePinBoneTransforms(InProcessor, InSourceGlobalPose, OutTargetGlobalPose);
 
 	// TODO: Should we just use initial chain lengths?
 	CacheChainLengthMap.Reset();
@@ -226,18 +212,26 @@ void FIKRetargetRelativeIKOp::RunAfterParent(
 	{
 		return;
 	}
-		
-	for (int32 Index=0; Index < CachePinSourcePropBones.Num(); ++Index)
+	
+	for (const FRelativeIKPinBoneData& PinData : CachePinBoneData)
 	{
-		FName SourcePinBone = CachePinSourcePropBones[Index]; 
-		FName TargetPinBone = ApplyBodyMap(SourcePinBone);
-		if (TargetPinBone == NAME_None || !CacheTargetSkelIndices.Contains(TargetPinBone))
+		FName SourcePinBone = PinData.SourceBone;
+		FName TargetPinBone = PinData.TargetBone;
+		if (!CacheSourceSkelIndices.Contains(SourcePinBone) || !CacheTargetSkelIndices.Contains(TargetPinBone))
 		{
 			continue;
 		}
-
+		
 		int32 TargetBoneIdx = CacheTargetSkelIndices[TargetPinBone];
-		OutTargetGlobalPose[TargetBoneIdx] = CachePinTargetPropTransforms[Index];
+		int32 CachePropIdx = CachePinSourcePropBones.Find(SourcePinBone);
+		if (CachePropIdx != INDEX_NONE)
+		{
+			OutTargetGlobalPose[TargetBoneIdx] = CachePinTargetPropTransforms[CachePropIdx];
+		}
+		else
+		{
+			OutTargetGlobalPose[TargetBoneIdx] = ComputePinBoneTransform(PinData, InProcessor, InSourceGlobalPose, OutTargetGlobalPose);
+		}
 	}
 }
 
@@ -580,6 +574,206 @@ void FIKRetargetRelativeIKOp::DropPropDomainBodyEffects()
 	}
 }
 
+void FIKRetargetRelativeIKOp::UpdatePinData(const FIKRetargetProcessor& Processor)
+{
+	// get skeletons we are copying from/to
+	const FRetargetSkeleton& SourceSkeleton = Processor.GetSkeleton(ERetargetSourceOrTarget::Source);
+	const FRetargetSkeleton& TargetSkeleton = Processor.GetSkeleton(ERetargetSourceOrTarget::Target);
+
+	// get retarget pose of source and target
+	const TArray<FTransform>& SourceRetargetRefPose = SourceSkeleton.RetargetPoses.GetGlobalRetargetPose();
+	const TArray<FTransform>& TargetRetargeRefPose = TargetSkeleton.RetargetPoses.GetGlobalRetargetPose();
+	
+	CachePinBoneData.Reset(Settings.PinPropBones.Num());
+	for (FName PropBone : Settings.PinPropBones)
+	{
+		FName TargetBoneName = ApplyBodyMap(PropBone);
+		if (TargetBoneName == NAME_None)
+		{
+			continue;
+		}
+		
+		UpdateSourceBoneMap(PropBone);
+		UpdateTargetBoneMap(TargetBoneName);
+		
+		if (!CacheSourceSkelIndices.Contains(PropBone) || !CacheTargetSkelIndices.Contains(TargetBoneName))
+		{
+			continue;
+		}
+		
+		// get relevant transforms from the retarget poses
+		int32 PropBoneIdx = CacheSourceSkelIndices[PropBone];
+		int32 TargetBoneIdx = CacheTargetSkelIndices[TargetBoneName];
+		
+		FTransform PropBoneParentRefPoseGlobal = GetParentTransform(SourceSkeleton, PropBoneIdx, SourceRetargetRefPose);
+		FTransform PropBoneRefPoseGlobal = SourceRetargetRefPose[PropBoneIdx];
+		
+		FTransform TargetBoneParentRefPoseGlobal = GetParentTransform(TargetSkeleton, TargetBoneIdx, TargetRetargeRefPose);
+		FTransform TargetBoneRefPoseGlobal = TargetRetargeRefPose[TargetBoneIdx];
+		
+		FRelativeIKPinBoneData PinData;
+		PinData.SourceBone = PropBone;
+		PinData.TargetBone = TargetBoneName;
+		
+		PinData.SourceTargetOffsetTfm = PropBoneRefPoseGlobal.GetRelativeTransform(TargetBoneRefPoseGlobal);
+		PinData.SourceLocalRefTfm = PropBoneParentRefPoseGlobal.GetRelativeTransform(PropBoneRefPoseGlobal);
+		PinData.TargetLocalRefTfm = TargetBoneParentRefPoseGlobal.GetRelativeTransform(TargetBoneRefPoseGlobal);
+		PinData.RefRotDelta = TargetBoneRefPoseGlobal.GetRotation() * PropBoneRefPoseGlobal.GetRotation().Inverse();
+		
+		CachePinBoneData.Add(PinData);
+	}
+}
+
+void FIKRetargetRelativeIKOp::CachePinBoneTransforms(
+	const FIKRetargetProcessor& Processor,
+	const TArray<FTransform>& InSourceGlobalPose,
+	TArray<FTransform>& OutTargetGlobalPose)
+{
+	CachePinTargetPropTransforms.SetNumZeroed(CachePinSourcePropBones.Num());
+	for (const FRelativeIKPinBoneData& PinData : CachePinBoneData)
+	{
+		FName SourcePropBone = PinData.SourceBone;
+		// Will calculate non-cached pins in ApplyAfterParent
+		int32 CachePropIdx = CachePinSourcePropBones.Find(SourcePropBone);
+		if (CachePropIdx == INDEX_NONE)
+		{
+			continue;
+		}
+		
+		FName TargetPropBone = ApplyBodyMap(SourcePropBone);
+		if (!CacheSourceSkelIndices.Contains(SourcePropBone) || !CacheTargetSkelIndices.Contains(TargetPropBone))
+		{
+			continue;
+		}
+		
+		int32 TargetBoneIdx = CacheTargetSkelIndices[TargetPropBone];
+		FTransform OutGlobalTransform = ComputePinBoneTransform(PinData, Processor, InSourceGlobalPose, OutTargetGlobalPose);
+		
+		OutTargetGlobalPose[TargetBoneIdx] = OutGlobalTransform;
+		CachePinTargetPropTransforms[CachePropIdx] = OutGlobalTransform;
+	}
+}
+
+FTransform FIKRetargetRelativeIKOp::ComputePinBoneTransform(
+		const FRelativeIKPinBoneData& PinData,
+		const FIKRetargetProcessor& Processor,
+		const TArray<FTransform>& InSourceGlobalPose,
+		const TArray<FTransform>& OutTargetGlobalPose)
+{
+	int32 SourceBoneIdx = CacheSourceSkelIndices[PinData.SourceBone];
+	int32 TargetBoneIdx = CacheTargetSkelIndices[PinData.TargetBone];
+	
+	const FRetargetSkeleton& SourceSkeleton = Processor.GetSkeleton(ERetargetSourceOrTarget::Source);
+	const FRetargetSkeleton& TargetSkeleton = Processor.GetSkeleton(ERetargetSourceOrTarget::Target);
+	
+	const FTransform& SourceTransform = InSourceGlobalPose[SourceBoneIdx]; 
+	FTransform TargetTransform = SourceTransform;
+	
+	// get the current translational offset of BoneToCopyFrom relative to it's parent
+	auto GetSourceLocalBoneVector = [SourceSkeleton, SourceBoneIdx, InSourceGlobalPose]()
+	{
+		const FTransform& SourceTransform = InSourceGlobalPose[SourceBoneIdx];
+		const FTransform& ParentTransform = GetParentTransform(SourceSkeleton, SourceBoneIdx, InSourceGlobalPose);
+		return SourceTransform.GetTranslation() - ParentTransform.GetTranslation();
+	};
+	
+	// generate translational offset
+	switch (Settings.TranslationMode)
+	{
+	case EPinBoneTranslationMode::CopyGlobalPosition:
+		{
+			TargetTransform.SetTranslation(SourceTransform.GetTranslation());
+			break;
+		}
+		
+	case EPinBoneTranslationMode::CopyGlobalPositionAndMaintainOffset:
+		{
+			TargetTransform.SetTranslation((PinData.SourceTargetOffsetTfm * SourceTransform).GetTranslation());
+			break;
+		}
+		
+	case EPinBoneTranslationMode::CopyLocalPosition:
+		{
+			const FVector SourceLocalVector = GetSourceLocalBoneVector();
+			const FTransform& TargetParentTfm = GetParentTransform(TargetSkeleton, TargetBoneIdx, OutTargetGlobalPose);
+			TargetTransform.SetTranslation(TargetParentTfm.GetTranslation() + SourceLocalVector);
+			break;
+		}
+		
+	case EPinBoneTranslationMode::CopyLocalPositionRelativeOffset:
+		{
+			const FVector SourceLocalVector = GetSourceLocalBoneVector();
+			double SourceBoneLength;
+			FVector SourceTranslationDir;
+			SourceLocalVector.ToDirectionAndLength(SourceTranslationDir, SourceBoneLength);
+			
+			const double RestPoseLengthDifference = PinData.TargetLocalRefTfm.GetTranslation().Size() - PinData.SourceLocalRefTfm.GetTranslation().Size();
+			const FVector SourceRelativeOffset = SourceTranslationDir * (SourceBoneLength + RestPoseLengthDifference);
+		
+			const FTransform& ParentOfBoneToCopyToCurrent = GetParentTransform(TargetSkeleton, TargetBoneIdx, OutTargetGlobalPose);
+			TargetTransform.SetTranslation(ParentOfBoneToCopyToCurrent.GetTranslation() + SourceRelativeOffset);
+			break;
+		}
+		
+	case EPinBoneTranslationMode::CopyLocalPositionRelativeScaled:
+		{
+			const FVector SourceLocalVector = GetSourceLocalBoneVector();
+			double SourceBoneLength;
+			FVector SourceTranslationDir;
+			SourceLocalVector.ToDirectionAndLength(SourceTranslationDir, SourceBoneLength);
+			
+			const double SourceRefBoneLength = PinData.SourceLocalRefTfm.GetTranslation().Size();
+			const double TargetRefBoneLength = PinData.TargetLocalRefTfm.GetTranslation().Size();
+			const double BoneLengthScaleFactor = TargetRefBoneLength / SourceRefBoneLength;
+			
+			const FVector BoneToCopyFromRelativeOffset = SourceTranslationDir * (SourceBoneLength * BoneLengthScaleFactor);
+		
+			const FTransform& ParentOfBoneToCopyToCurrent = GetParentTransform(TargetSkeleton, TargetBoneIdx, OutTargetGlobalPose);
+			TargetTransform.SetTranslation(ParentOfBoneToCopyToCurrent.GetTranslation() + BoneToCopyFromRelativeOffset);
+			
+			break;
+		}
+		
+	default:
+		ensureAlwaysMsgf(false, TEXT("Missing translation offset mode."));
+		break;
+	}
+
+	// generate rotational offset
+	switch (Settings.RotationMode)
+	{
+	case EPinBoneRotationMode::CopyGlobalRotation:
+		{
+			TargetTransform.SetRotation(SourceTransform.GetRotation());
+			break;
+		}
+	case EPinBoneRotationMode::MaintainOffsetFromBoneToCopyFrom:
+		{
+			// apply "copy from" rotation plus delta
+			TargetTransform.SetRotation(PinData.RefRotDelta * SourceTransform.GetRotation());
+			
+			break;
+		}
+		
+	default:
+		ensureAlwaysMsgf(false, TEXT("Missing rotation offset mode."));
+		break;
+	}
+	
+	return TargetTransform;
+}
+
+const FTransform& FIKRetargetRelativeIKOp::GetParentTransform(const FRetargetSkeleton& Skeleton, const int32 BoneIndex, const TArray<FTransform>& InPose)
+{
+	int32 ParentIdx = Skeleton.GetParentIndex(BoneIndex);
+	if (ParentIdx == INDEX_NONE)
+	{
+		return FTransform::Identity;
+	}
+	
+	return InPose[ParentIdx];
+}
+
 void FIKRetargetRelativeIKOp::ApplyAverageBodyPairTargets(
 	TArray<FIKRigGoal>& OutIKGoals,
 	const TArray<FTransform>& TargetPose,
diff --git a/Engine/Plugins/Experimental/Animation/RelativeIKOp/Source/RelativeIKOp/Public/RelativeIKOp.h b/Engine/Plugins/Experimental/Animation/RelativeIKOp/Source/RelativeIKOp/Public/RelativeIKOp.h
--- a/Engine/Plugins/Experimental/Animation/RelativeIKOp/Source/RelativeIKOp/Public/RelativeIKOp.h
+++ b/Engine/Plugins/Experimental/Animation/RelativeIKOp/Source/RelativeIKOp/Public/RelativeIKOp.h
@@ -2,6 +2,7 @@
 #pragma once
 
 #include "Retargeter/IKRetargetOps.h"
+#include "Retargeter/RetargetOps/PinBoneOp.h"
 #include "Eigen/Dense"
 
 #include "RelativeIKOp.generated.h"
@@ -100,6 +101,17 @@ struct FDebugRelativeIKDrawInfo
 	TArray<FName> TargetDomainBones;
 };
 
+struct FRelativeIKPinBoneData
+{
+	FName SourceBone;
+	FName TargetBone;
+	
+	FTransform SourceTargetOffsetTfm;
+	FTransform SourceLocalRefTfm;
+	FTransform TargetLocalRefTfm;
+	FQuat RefRotDelta;
+};
+
 USTRUCT(BlueprintType, meta = (DisplayName = "Relative IK Settings"))
 struct FIKRetargetRelativeIKOpSettings : public FIKRetargetOpSettingsBase
 {
@@ -121,6 +133,20 @@ struct FIKRetargetRelativeIKOpSettings : public FIKRetargetOpSettingsBase
 	UPROPERTY(EditAnywhere, BlueprintReadWrite, Category = "Source Target Mapping", meta=(ReinitializeOnEdit))
 	TMap<FName,FName> BodyMapping;
 	
+	//TODO: Is there a better way to replicate the pinning functionality
+	
+	// List of bones to pin. Should include all attach bones that may be affected by relative ik bake data 
+	UPROPERTY(EditAnywhere, BlueprintReadWrite, Category = "Prop Pinning", meta=(ReinitializeOnEdit))
+	TArray<FName> PinPropBones;
+	
+	// The method used to calculate the translation of the bone.
+	UPROPERTY(EditAnywhere, BlueprintReadWrite, Category = "Prop Pinning")
+	EPinBoneTranslationMode TranslationMode = EPinBoneTranslationMode::CopyGlobalPosition;
+
+	// The method used to calculate the rotation of the bone.
+	UPROPERTY(EditAnywhere, BlueprintReadWrite, Category = "Prop Pinning")
+	EPinBoneRotationMode RotationMode = EPinBoneRotationMode::CopyGlobalRotation;
+	
 	// Maximum distance for which body pair info is baked
 	UPROPERTY(EditAnywhere, BlueprintReadWrite, Category = "Parameters", meta=(ClampMin = "0.1"))
 	double DistanceThreshold = 50.00;
@@ -264,13 +290,26 @@ private:
 	void SetupRelativeIKNotifyInfoMontage();
 	void UpdateRelativeIKNotifyInfoAnimSeq(UAnimInstance* SourceAnimInstance);
 	void ResetCacheNotifyInfo();
-	void PostCachedAllNotifyInfo();
 	void UpdateCacheNotifyInfo(URelativeBodyBakeAnimNotify* NotifyInfo);
 	void UpdateCacheChainInfo(const TArray<FTransform>& TargetPose, FName BoneName);
 	void UpdateCachePropNotifyInfo(URelativePropsBakeAnimNotify* NotifyInfo);
 
 	void DropPropDomainBodyEffects();
 	
+	void UpdatePinData(const FIKRetargetProcessor& Processor);
+	void CachePinBoneTransforms(
+		const FIKRetargetProcessor& Processor,
+		const TArray<FTransform>& InSourceGlobalPose,
+		TArray<FTransform>& OutTargetGlobalPose);
+	
+	FTransform ComputePinBoneTransform(
+		const FRelativeIKPinBoneData& PinData,
+		const FIKRetargetProcessor& Processor,
+		const TArray<FTransform>& InSourceGlobalPose,
+		const TArray<FTransform>& OutTargetGlobalPose);
+	
+	static const FTransform& GetParentTransform(const FRetargetSkeleton& Skeleton, const int32 BoneIndex, const TArray<FTransform>& InPose);
+	
 	void ComputeFramePairSpaces(
 		double SourceScale,
 		const TArray<FTransform>& SourcePose,
@@ -449,6 +488,7 @@ private:
 	TMap<FName,FVector> CacheChainStartMap;
 	// Cache prop bones to pin/update
 	TArray<FName> CachePinSourcePropBones;
+	TArray<FRelativeIKPinBoneData> CachePinBoneData;
 	TArray<FTransform> CachePinTargetPropTransforms;
 
 	// Cached montage instance/segment info for early-out matching in preupdate
