RelativeIKOp: Per-bone pin options and fix for local translation scaled offset copy.

#Synced-CL 48387090

diff --git a/Engine/Plugins/Experimental/Animation/RelativeIKOp/Source/RelativeIKOp/Private/RelativeIKOp.cpp b/Engine/Plugins/Experimental/Animation/RelativeIKOp/Source/RelativeIKOp/Private/RelativeIKOp.cpp
--- a/Engine/Plugins/Experimental/Animation/RelativeIKOp/Source/RelativeIKOp/Private/RelativeIKOp.cpp
+++ b/Engine/Plugins/Experimental/Animation/RelativeIKOp/Source/RelativeIKOp/Private/RelativeIKOp.cpp
@@ -585,24 +585,25 @@ void FIKRetargetRelativeIKOp::UpdatePinData(const FIKRetargetProcessor& Processo
 	const TArray<FTransform>& TargetRetargeRefPose = TargetSkeleton.RetargetPoses.GetGlobalRetargetPose();
 	
 	CachePinBoneData.Reset(Settings.PinPropBones.Num());
-	for (FName PropBone : Settings.PinPropBones)
+	for (const FPinBoneSettings& PinBoneInfo : Settings.PinPropBones)
 	{
-		FName TargetBoneName = ApplyBodyMap(PropBone);
+		FName SourceBoneName = PinBoneInfo.SourceBoneName;
+		FName TargetBoneName = ApplyBodyMap(SourceBoneName);
 		if (TargetBoneName == NAME_None)
 		{
 			continue;
 		}
 		
-		UpdateSourceBoneMap(PropBone);
+		UpdateSourceBoneMap(SourceBoneName);
 		UpdateTargetBoneMap(TargetBoneName);
 		
-		if (!CacheSourceSkelIndices.Contains(PropBone) || !CacheTargetSkelIndices.Contains(TargetBoneName))
+		if (!CacheSourceSkelIndices.Contains(SourceBoneName) || !CacheTargetSkelIndices.Contains(TargetBoneName))
 		{
 			continue;
 		}
 		
 		// get relevant transforms from the retarget poses
-		int32 PropBoneIdx = CacheSourceSkelIndices[PropBone];
+		int32 PropBoneIdx = CacheSourceSkelIndices[SourceBoneName];
 		int32 TargetBoneIdx = CacheTargetSkelIndices[TargetBoneName];
 		
 		FTransform PropBoneParentRefPoseGlobal = GetParentTransform(SourceSkeleton, PropBoneIdx, SourceRetargetRefPose);
@@ -612,7 +613,10 @@ void FIKRetargetRelativeIKOp::UpdatePinData(const FIKRetargetProcessor& Processo
 		FTransform TargetBoneRefPoseGlobal = TargetRetargeRefPose[TargetBoneIdx];
 		
 		FRelativeIKPinBoneData PinData;
-		PinData.SourceBone = PropBone;
+		PinData.TranslationMode = PinBoneInfo.TranslationMode;
+		PinData.RotationMode = PinBoneInfo.RotationMode;
+		
+		PinData.SourceBone = SourceBoneName;
 		PinData.TargetBone = TargetBoneName;
 		
 		PinData.SourceTargetOffsetTfm = PropBoneRefPoseGlobal.GetRelativeTransform(TargetBoneRefPoseGlobal);
@@ -678,7 +682,7 @@ FTransform FIKRetargetRelativeIKOp::ComputePinBoneTransform(
 	};
 	
 	// generate translational offset
-	switch (Settings.TranslationMode)
+	switch (PinData.TranslationMode)
 	{
 	case EPinBoneTranslationMode::CopyGlobalPosition:
 		{
@@ -724,13 +728,12 @@ FTransform FIKRetargetRelativeIKOp::ComputePinBoneTransform(
 			
 			const double SourceRefBoneLength = PinData.SourceLocalRefTfm.GetTranslation().Size();
 			const double TargetRefBoneLength = PinData.TargetLocalRefTfm.GetTranslation().Size();
-			const double BoneLengthScaleFactor = TargetRefBoneLength / SourceRefBoneLength;
+			const double BoneLengthScaleFactor = (!FMath::IsNearlyZero(SourceRefBoneLength)) ? TargetRefBoneLength / SourceRefBoneLength : 0.0;
 			
 			const FVector BoneToCopyFromRelativeOffset = SourceTranslationDir * (SourceBoneLength * BoneLengthScaleFactor);
 		
 			const FTransform& ParentOfBoneToCopyToCurrent = GetParentTransform(TargetSkeleton, TargetBoneIdx, OutTargetGlobalPose);
 			TargetTransform.SetTranslation(ParentOfBoneToCopyToCurrent.GetTranslation() + BoneToCopyFromRelativeOffset);
-			
 			break;
 		}
 		
@@ -740,7 +743,7 @@ FTransform FIKRetargetRelativeIKOp::ComputePinBoneTransform(
 	}
 
 	// generate rotational offset
-	switch (Settings.RotationMode)
+	switch (PinData.RotationMode)
 	{
 	case EPinBoneRotationMode::CopyGlobalRotation:
 		{
@@ -751,7 +754,6 @@ FTransform FIKRetargetRelativeIKOp::ComputePinBoneTransform(
 		{
 			// apply "copy from" rotation plus delta
 			TargetTransform.SetRotation(PinData.RefRotDelta * SourceTransform.GetRotation());
-			
 			break;
 		}
 		
diff --git a/Engine/Plugins/Experimental/Animation/RelativeIKOp/Source/RelativeIKOp/Public/RelativeIKOp.h b/Engine/Plugins/Experimental/Animation/RelativeIKOp/Source/RelativeIKOp/Public/RelativeIKOp.h
--- a/Engine/Plugins/Experimental/Animation/RelativeIKOp/Source/RelativeIKOp/Public/RelativeIKOp.h
+++ b/Engine/Plugins/Experimental/Animation/RelativeIKOp/Source/RelativeIKOp/Public/RelativeIKOp.h
@@ -101,15 +101,21 @@ struct FDebugRelativeIKDrawInfo
 	TArray<FName> TargetDomainBones;
 };
 
-struct FRelativeIKPinBoneData
+USTRUCT(BlueprintType)
+struct FPinBoneSettings
 {
-	FName SourceBone;
-	FName TargetBone;
+	GENERATED_BODY()
+	
+	UPROPERTY(EditAnywhere, BlueprintReadWrite, Category = "Pinning", meta=(ReinitializeOnEdit))
+	FName SourceBoneName;
+	
+	// The method used to calculate the translation of the bone.
+	UPROPERTY(EditAnywhere, BlueprintReadWrite, Category = "Pinning", meta=(ReinitializeOnEdit))
+	EPinBoneTranslationMode TranslationMode = EPinBoneTranslationMode::CopyGlobalPosition;
 	
-	FTransform SourceTargetOffsetTfm;
-	FTransform SourceLocalRefTfm;
-	FTransform TargetLocalRefTfm;
-	FQuat RefRotDelta;
+	// The method used to calculate the rotation of the bone.
+	UPROPERTY(EditAnywhere, BlueprintReadWrite, Category = "Pinning", meta=(ReinitializeOnEdit))
+	EPinBoneRotationMode RotationMode = EPinBoneRotationMode::CopyGlobalRotation;
 };
 
 USTRUCT(BlueprintType, meta = (DisplayName = "Relative IK Settings"))
@@ -121,10 +127,6 @@ struct FIKRetargetRelativeIKOpSettings : public FIKRetargetOpSettingsBase
 
 	UE_API virtual void CopySettingsAtRuntime(const FIKRetargetOpSettingsBase* InSettingsToCopyFrom) override;
 
-	// Source mesh physics asset for relative body pair tests 
-	// UPROPERTY(EditAnywhere, BlueprintReadWrite, Category = "Physics", meta=(ReinitializeOnEdit))
-	// TObjectPtr<UPhysicsAsset> SourcePhysicsAssetOverride;
-
 	// Target mesh physics asset for retargeting relative body pair vertices
 	UPROPERTY(EditAnywhere, BlueprintReadWrite, Category = "Physics", meta=(ReinitializeOnEdit))
 	TObjectPtr<UPhysicsAsset> TargetPhysicsAssetOverride;
@@ -133,19 +135,9 @@ struct FIKRetargetRelativeIKOpSettings : public FIKRetargetOpSettingsBase
 	UPROPERTY(EditAnywhere, BlueprintReadWrite, Category = "Source Target Mapping", meta=(ReinitializeOnEdit))
 	TMap<FName,FName> BodyMapping;
 	
-	//TODO: Is there a better way to replicate the pinning functionality
-	
 	// List of bones to pin. Should include all attach bones that may be affected by relative ik bake data 
 	UPROPERTY(EditAnywhere, BlueprintReadWrite, Category = "Prop Pinning", meta=(ReinitializeOnEdit))
-	TArray<FName> PinPropBones;
-	
-	// The method used to calculate the translation of the bone.
-	UPROPERTY(EditAnywhere, BlueprintReadWrite, Category = "Prop Pinning")
-	EPinBoneTranslationMode TranslationMode = EPinBoneTranslationMode::CopyGlobalPosition;
-
-	// The method used to calculate the rotation of the bone.
-	UPROPERTY(EditAnywhere, BlueprintReadWrite, Category = "Prop Pinning")
-	EPinBoneRotationMode RotationMode = EPinBoneRotationMode::CopyGlobalRotation;
+	TArray<FPinBoneSettings> PinPropBones;
 	
 	// Maximum distance for which body pair info is baked
 	UPROPERTY(EditAnywhere, BlueprintReadWrite, Category = "Parameters", meta=(ClampMin = "0.1"))
@@ -296,6 +288,20 @@ private:
 
 	void DropPropDomainBodyEffects();
 	
+	struct FRelativeIKPinBoneData
+	{
+		EPinBoneTranslationMode TranslationMode;
+		EPinBoneRotationMode RotationMode;
+		// TODO: Switch to using bone indices where possible
+		FName SourceBone;
+		FName TargetBone;
+	
+		FTransform SourceTargetOffsetTfm;
+		FTransform SourceLocalRefTfm;
+		FTransform TargetLocalRefTfm;
+		FQuat RefRotDelta;
+	};
+	
 	void UpdatePinData(const FIKRetargetProcessor& Processor);
 	void CachePinBoneTransforms(
 		const FIKRetargetProcessor& Processor,
