RelativeBodyAnimModifiers: Add prop bulk bake function for EUW/BP access. And fixed bh caught memleak in body bake code.

#Synced-CL 48450904

diff --git a/Engine/Plugins/Experimental/Animation/RelativeIKOp/Source/RelativeBodyAnimUtils/Private/RelativeBodyAnimModifier.cpp b/Engine/Plugins/Experimental/Animation/RelativeIKOp/Source/RelativeBodyAnimUtils/Private/RelativeBodyAnimModifier.cpp
--- a/Engine/Plugins/Experimental/Animation/RelativeIKOp/Source/RelativeBodyAnimUtils/Private/RelativeBodyAnimModifier.cpp
+++ b/Engine/Plugins/Experimental/Animation/RelativeIKOp/Source/RelativeBodyAnimUtils/Private/RelativeBodyAnimModifier.cpp
@@ -3,7 +3,6 @@
 #include "RelativeBodyAnimModifier.h"
 
 #include "AnimPose.h"
-#include "AnimationModifiersAssetUserData.h"
 #include "MeshDescription.h"
 #include "RelativeBodyAnimNotifies.h"
 #include "RelativeBodyUtils.h"
@@ -16,71 +15,10 @@
 #include "PhysicsEngine/SkeletalBodySetup.h"
 #include "Rendering/SkeletalMeshRenderData.h"
 #include "Spatial/PointHashGrid3.h"
+#include "Templates/UniquePtr.h"
 
 #define LOCTEXT_NAMESPACE "RelativeBodyAnimModifier"
 
-URelativeBodyAnimModifier* FindRelativeBodyModifier(UAnimationModifiersAssetUserData* AssetUserData)
-{
-	const TArray<UAnimationModifier*>& ModifierInstances = AssetUserData->GetAnimationModifierInstances();
-	for (UAnimationModifier* Modifier : ModifierInstances)
-	{
-		if (URelativeBodyAnimModifier* CheckInst = Cast<URelativeBodyAnimModifier>(Modifier))
-		{
-			return CheckInst;
-		}
-	}
-
-	return nullptr;
-}
-
-void URelativeBodyAnimBlueprintFunctions::BulkUpdateRelativeBodyModifiers(TArray<UAnimSequenceBase*>& AnimationList, const FRelativeBodyAnimModifierOptions& Options)
-{
-	FScopedSlowTask ApplyModifiersTask(static_cast<float>(AnimationList.Num()), LOCTEXT("BulkUpdateProgress", "Updating and applying modifier to animation list"));
-	ApplyModifiersTask.MakeDialog();
-	for (UAnimSequenceBase* Anim : AnimationList)
-	{
-		ApplyModifiersTask.EnterProgressFrame(1.0f, FText::FromString(Anim->GetName()));
-		UAnimationModifiersAssetUserData* AssetUserData = Anim->GetAssetUserData<UAnimationModifiersAssetUserData>();
-		if (!AssetUserData)
-		{
-			UE_LOG(LogAnimation, Warning, TEXT("Unable to find modifier user data for anim: %s"), *Anim->GetName());
-			continue;
-		}
-
-		TObjectPtr<URelativeBodyAnimModifier> UpdateModifierInst = FindRelativeBodyModifier(AssetUserData);
-		if (!UpdateModifierInst)
-		{
-			if (!AssetUserData->AddAnimationModifierOfClass(Anim, URelativeBodyAnimModifier::StaticClass()))
-			{
-				UE_LOG(LogAnimation, Warning, TEXT("Unable to create new modifier for anim: %s"), *Anim->GetName());
-				continue;
-			}
-
-			UpdateModifierInst = FindRelativeBodyModifier(AssetUserData);
-		}
-
-		if (!UpdateModifierInst)
-		{
-			continue;
-		}
-
-		UpdateModifierInst->SampleRate = Options.SampleRate;
-		UpdateModifierInst->ContactThreshold = Options.ContactThreshold;
-		UpdateModifierInst->SkeletalMeshAsset = Options.SkeletalMeshAsset;
-		UpdateModifierInst->PhysicsAssetOverride = Options.PhysicsAssetOverride;
-		UpdateModifierInst->DomainBodyNames = Options.DomainBodyNames;
-		UpdateModifierInst->ContactBodyNames = Options.ContactBodyNames;
-		UpdateModifierInst->NotifyClass = Options.NotifyClass;
-
-		UAnimSequence* AnimSeq = Cast<UAnimSequence>(Anim);
-		if (AnimSeq)
-		{
-			UpdateModifierInst->ApplyToAnimationSequence(AnimSeq);
-		}
-	}
-}
-
-
 void URelativeBodyAnimModifier::OnApply_Implementation(UAnimSequence* InAnimation)
 {
 	// TODO: Make log category for relative body utils?
@@ -260,12 +198,12 @@ void URelativeBodyAnimModifier::OnApply_Implementation(UAnimSequence* InAnimatio
 			GetSkinnedVertices(VLocations, CacheToLocals);
 
 			// build a spatial hash grid
-			TArray<UE::Geometry::TPointHashGrid3f<int32>*> BodiesVertHash;
+			TArray<TUniquePtr<UE::Geometry::TPointHashGrid3f<int32>>> BodiesVertHash;
 			BodiesVertHash.SetNum(NumBodies);
 			size_t count = 0;
 			for (int32 BodyIndex = 0; BodyIndex < NumBodies; ++BodyIndex)
 			{
-				BodiesVertHash[BodyIndex] = new UE::Geometry::TPointHashGrid3f<int32>(HashGridCellSize, INDEX_NONE);
+				BodiesVertHash[BodyIndex] = MakeUnique<UE::Geometry::TPointHashGrid3f<int32>>(HashGridCellSize, INDEX_NONE);
 				BodiesVertHash[BodyIndex]->Reserve(VertexIndicesInfluencedByBodyIndices[BodyIndex].Num());
 				for (int32 VertexIndex : VertexIndicesInfluencedByBodyIndices[BodyIndex])
 				{
diff --git a/Engine/Plugins/Experimental/Animation/RelativeIKOp/Source/RelativeBodyAnimUtils/Private/RelativeBodyBlueprintFunctions.cpp b/Engine/Plugins/Experimental/Animation/RelativeIKOp/Source/RelativeBodyAnimUtils/Private/RelativeBodyBlueprintFunctions.cpp
new file mode 100644
--- /dev/null
+++ b/Engine/Plugins/Experimental/Animation/RelativeIKOp/Source/RelativeBodyAnimUtils/Private/RelativeBodyBlueprintFunctions.cpp
@@ -0,0 +1,178 @@
+﻿// Copyright Epic Games, Inc. All Rights Reserved.
+
+#include "RelativeBodyBlueprintFunctions.h"
+#include "AnimationModifiersAssetUserData.h"
+#include "RelativeBodyAnimNotifies.h"
+#include "Animation/AnimSequence.h"
+#include "Misc/ScopedSlowTask.h"
+
+#define LOCTEXT_NAMESPACE "RelativeBodyBlueprintFunctions"
+
+template <typename TRelativeBodyClass>
+TRelativeBodyClass* FindRelativeBodyModifier(UAnimationModifiersAssetUserData* AssetUserData)
+{
+	const TArray<UAnimationModifier*>& ModifierInstances = AssetUserData->GetAnimationModifierInstances();
+	for (UAnimationModifier* Modifier : ModifierInstances)
+	{
+		if (TRelativeBodyClass* CheckInst = Cast<TRelativeBodyClass>(Modifier))
+		{
+			return CheckInst;
+		}
+	}
+
+	return nullptr;
+}
+
+template <typename TRelativeBodyClass>
+TRelativeBodyClass* FindExactTypeBodyModifier(UAnimationModifiersAssetUserData* AssetUserData)
+{
+	const TArray<UAnimationModifier*>& ModifierInstances = AssetUserData->GetAnimationModifierInstances();
+	for (UAnimationModifier* Modifier : ModifierInstances)
+	{
+		if (TRelativeBodyClass* CheckInst = ExactCast<TRelativeBodyClass>(Modifier))
+		{
+			return CheckInst;
+		}
+	}
+
+	return nullptr;
+}
+
+void URelativeBodyAnimBlueprintFunctions::BulkUpdateRelativeBodyModifiers(TArray<UAnimSequenceBase*>& AnimationList, const FRelativeBodyAnimModifierOptions& Options)
+{
+	FScopedSlowTask ApplyModifiersTask(static_cast<float>(AnimationList.Num()), LOCTEXT("BulkUpdateProgress", "Updating and applying modifier to animation list"));
+	ApplyModifiersTask.MakeDialog();
+	for (UAnimSequenceBase* Anim : AnimationList)
+	{
+		ApplyModifiersTask.EnterProgressFrame(1.0f, FText::FromString(Anim->GetName()));
+		UAnimationModifiersAssetUserData* AssetUserData = Anim->GetAssetUserData<UAnimationModifiersAssetUserData>();
+		if (!AssetUserData)
+		{
+			UE_LOG(LogAnimation, Warning, TEXT("Unable to find modifier user data for anim: %s"), *Anim->GetName());
+			continue;
+		}
+
+		TObjectPtr<URelativeBodyAnimModifier> UpdateModifierInst = FindRelativeBodyModifier<URelativeBodyAnimModifier>(AssetUserData);
+		if (!UpdateModifierInst)
+		{
+			if (!AssetUserData->AddAnimationModifierOfClass(Anim, URelativeBodyAnimModifier::StaticClass()))
+			{
+				UE_LOG(LogAnimation, Warning, TEXT("Unable to create new modifier for anim: %s"), *Anim->GetName());
+				continue;
+			}
+
+			UpdateModifierInst = FindRelativeBodyModifier<URelativeBodyAnimModifier>(AssetUserData);
+		}
+
+		if (!UpdateModifierInst)
+		{
+			continue;
+		}
+
+		UpdateModifierInst->SampleRate = Options.SampleRate;
+		UpdateModifierInst->ContactThreshold = Options.ContactThreshold;
+		UpdateModifierInst->SkeletalMeshAsset = Options.SkeletalMeshAsset;
+		UpdateModifierInst->PhysicsAssetOverride = Options.PhysicsAssetOverride;
+		UpdateModifierInst->DomainBodyNames = Options.DomainBodyNames;
+		UpdateModifierInst->ContactBodyNames = Options.ContactBodyNames;
+		UpdateModifierInst->NotifyClass = Options.NotifyClass;
+
+		UAnimSequence* AnimSeq = Cast<UAnimSequence>(Anim);
+		if (AnimSeq)
+		{
+			UpdateModifierInst->ApplyToAnimationSequence(AnimSeq);
+		}
+	}
+}
+
+void URelativeBodyAnimBlueprintFunctions::BulkUpdatePropBodyModifiers(UPARAM(ref) TArray<UAnimSequenceBase*>& AnimationList, const FPropBodyAnimModifierOptions& Options)
+{
+	// Just run body mod update if the options don't contain prop info
+	if (Options.PropNotifyClass == nullptr || Options.PropMeshSettings.IsEmpty())
+	{
+		BulkUpdateRelativeBodyModifiers(AnimationList, Options);
+		return;
+	}
+	
+	FScopedSlowTask ApplyModifiersTask(static_cast<float>(AnimationList.Num()), LOCTEXT("BulkUpdateProgress", "Updating and applying prop/body modifiers to animation list"));
+	ApplyModifiersTask.MakeDialog();
+	for (UAnimSequenceBase* Anim : AnimationList)
+	{
+		ApplyModifiersTask.EnterProgressFrame(1.0f, FText::FromString(Anim->GetName()));
+		UAnimationModifiersAssetUserData* AssetUserData = Anim->GetAssetUserData<UAnimationModifiersAssetUserData>();
+		if (!AssetUserData)
+		{
+			UE_LOG(LogAnimation, Warning, TEXT("Unable to find modifier user data for anim: %s"), *Anim->GetName());
+			continue;
+		}
+
+		TObjectPtr<URelativePropsAnimModifier> UpdateModifierInst = FindRelativeBodyModifier<URelativePropsAnimModifier>(AssetUserData);
+		if (!UpdateModifierInst)
+		{
+			if (!AssetUserData->AddAnimationModifierOfClass(Anim, URelativePropsAnimModifier::StaticClass()))
+			{
+				UE_LOG(LogAnimation, Warning, TEXT("Unable to create new modifier for anim: %s"), *Anim->GetName());
+				continue;
+			}
+
+			UpdateModifierInst = FindRelativeBodyModifier<URelativePropsAnimModifier>(AssetUserData);
+		}
+
+		if (!UpdateModifierInst)
+		{
+			continue;
+		}
+		
+		UpdateModifierInst->PropsStaticMeshs.Empty(Options.PropMeshSettings.Num());
+		UpdateModifierInst->PropsAttachBones.Empty(Options.PropMeshSettings.Num());
+		UpdateModifierInst->PropsOffsetTransforms.Empty(Options.PropMeshSettings.Num());
+		for (const FPropMeshSettings& PropSettings : Options.PropMeshSettings)
+		{
+			UpdateModifierInst->PropsStaticMeshs.Add(PropSettings.StaticMesh);
+			UpdateModifierInst->PropsAttachBones.Add(PropSettings.AttachBoneName);
+			UpdateModifierInst->PropsOffsetTransforms.Add(PropSettings.AttachTransform);
+		}
+		UpdateModifierInst->PropsNotifyClass = Options.PropNotifyClass;
+		
+		// Always set these
+		UpdateModifierInst->SampleRate = Options.SampleRate;
+		UpdateModifierInst->ContactThreshold = Options.ContactThreshold;
+		UpdateModifierInst->SkeletalMeshAsset = Options.SkeletalMeshAsset;
+		UpdateModifierInst->PhysicsAssetOverride = Options.PhysicsAssetOverride;
+		UpdateModifierInst->DomainBodyNames = Options.DomainBodyNames;
+		UpdateModifierInst->ContactBodyNames = Options.ContactBodyNames;
+		UpdateModifierInst->NotifyClass = nullptr;
+		
+		// TODO: Ideally we would just delete the old body modifier in this case
+		bool bUpdateSeparateBodyModifier = false;
+		TObjectPtr<URelativeBodyAnimModifier> BodyModifierInst = FindExactTypeBodyModifier<URelativeBodyAnimModifier>(AssetUserData);
+		if (Options.NotifyClass != nullptr && BodyModifierInst != nullptr)
+		{
+			// Update body modifier to generate the 
+			bUpdateSeparateBodyModifier = true;
+			BodyModifierInst->SampleRate = Options.SampleRate;
+			BodyModifierInst->ContactThreshold = Options.ContactThreshold;
+			BodyModifierInst->SkeletalMeshAsset = Options.SkeletalMeshAsset;
+			BodyModifierInst->PhysicsAssetOverride = Options.PhysicsAssetOverride;
+			BodyModifierInst->DomainBodyNames = Options.DomainBodyNames;
+			BodyModifierInst->ContactBodyNames = Options.ContactBodyNames;
+			BodyModifierInst->NotifyClass = Options.NotifyClass;
+		}
+		else
+		{
+			UpdateModifierInst->NotifyClass = Options.NotifyClass;
+		}
+
+		UAnimSequence* AnimSeq = Cast<UAnimSequence>(Anim);
+		if (AnimSeq)
+		{
+			UpdateModifierInst->ApplyToAnimationSequence(AnimSeq);
+			if (bUpdateSeparateBodyModifier)
+			{
+				BodyModifierInst->ApplyToAnimationSequence(AnimSeq);
+			}
+		}
+	}
+}
+
+#undef LOCTEXT_NAMESPACE
diff --git a/Engine/Plugins/Experimental/Animation/RelativeIKOp/Source/RelativeBodyAnimUtils/Public/RelativeBodyAnimModifier.h b/Engine/Plugins/Experimental/Animation/RelativeIKOp/Source/RelativeBodyAnimUtils/Public/RelativeBodyAnimModifier.h
--- a/Engine/Plugins/Experimental/Animation/RelativeIKOp/Source/RelativeBodyAnimUtils/Public/RelativeBodyAnimModifier.h
+++ b/Engine/Plugins/Experimental/Animation/RelativeIKOp/Source/RelativeBodyAnimUtils/Public/RelativeBodyAnimModifier.h
@@ -39,51 +39,6 @@ struct FRelativeBodySourceData
 	TArray<int32> SourceBoneIndicesToBodyIndices;
 };
 
-USTRUCT(BlueprintType)
-struct RELATIVEBODYANIMUTILS_API FRelativeBodyAnimModifierOptions
-{
-	GENERATED_BODY()
-
-	/** Rate used to sample the animation */
-	UPROPERTY(EditAnywhere, BlueprintReadWrite, Category = "Settings", meta = (Units = "Hz", UIMin = 1))
-	int SampleRate = 30;
-
-	/** Threshold for determining if a bone pair can be considered to be having contact */
-	UPROPERTY(EditAnywhere, BlueprintReadWrite, Category = "Settings")
-	float ContactThreshold = 150.0f;
-
-	UPROPERTY(EditAnywhere, BlueprintReadWrite, Category = "Settings", meta = (TitleProperty = "SkeletalMeshAsset"))
-	TObjectPtr<USkeletalMesh> SkeletalMeshAsset;
-
-	UPROPERTY(EditAnywhere, BlueprintReadWrite, Category = "Settings", meta = (TitleProperty = "PhysicsAssetOverride"))
-	TObjectPtr<UPhysicsAsset> PhysicsAssetOverride;
-
-	/** Bodies to be checked against contact bodies */
-	UPROPERTY(EditAnywhere, BlueprintReadWrite, Category = "Settings", meta = (TitleProperty = "DomainBodyNames"))
-	TArray<FName> DomainBodyNames;
-	
-	UPROPERTY(EditAnywhere, BlueprintReadWrite, Category = "Settings", meta = (TitleProperty = "ContactBodyNames"))
-	TArray<FName> ContactBodyNames;
-
-	// Relative Body Notify subclass to create
-	UPROPERTY(EditAnywhere, BlueprintReadWrite, Category = "Settings", meta = (TitleProperty = "PhysicalIKRetargeterNotifyClass"))
-	TSubclassOf<URelativeBodyAnimNotifyBase> NotifyClass = nullptr;
-};
-
-/**
- * BPFL helper for setting up anim modifier 
- */
-UCLASS(Transient, meta=(ScriptName="RelativeBodyAnimModifyHelper"))
-class RELATIVEBODYANIMUTILS_API URelativeBodyAnimBlueprintFunctions : public UBlueprintFunctionLibrary
-{
-	GENERATED_BODY()
-public:
-	
-	UFUNCTION(BlueprintCallable, Category = "Anim Modifier")
-	static void BulkUpdateRelativeBodyModifiers(UPARAM(ref) TArray<UAnimSequenceBase*>& AnimationList, const FRelativeBodyAnimModifierOptions& Options);
-};
-
-
 /**
  * Animation modifier for baking relative body relationships into anim notify
  */
diff --git a/Engine/Plugins/Experimental/Animation/RelativeIKOp/Source/RelativeBodyAnimUtils/Public/RelativeBodyBlueprintFunctions.h b/Engine/Plugins/Experimental/Animation/RelativeIKOp/Source/RelativeBodyAnimUtils/Public/RelativeBodyBlueprintFunctions.h
new file mode 100644
--- /dev/null
+++ b/Engine/Plugins/Experimental/Animation/RelativeIKOp/Source/RelativeBodyAnimUtils/Public/RelativeBodyBlueprintFunctions.h
@@ -0,0 +1,85 @@
+﻿// Copyright Epic Games, Inc. All Rights Reserved.
+
+#pragma once
+
+#include "RelativeBodyAnimModifier.h"
+#include "RelativePropsAnimModifier.h"
+
+#include "RelativeBodyBlueprintFunctions.generated.h"
+
+USTRUCT(BlueprintType)
+struct RELATIVEBODYANIMUTILS_API FRelativeBodyAnimModifierOptions
+{
+	GENERATED_BODY()
+
+	/** Rate used to sample the animation */
+	UPROPERTY(EditAnywhere, BlueprintReadWrite, Category = "Settings", meta = (Units = "Hz", UIMin = 1))
+	int SampleRate = 30;
+
+	/** Threshold for determining if a bone pair can be considered to be having contact */
+	UPROPERTY(EditAnywhere, BlueprintReadWrite, Category = "Settings")
+	float ContactThreshold = 150.0f;
+
+	UPROPERTY(EditAnywhere, BlueprintReadWrite, Category = "Settings", meta = (TitleProperty = "SkeletalMeshAsset"))
+	TObjectPtr<USkeletalMesh> SkeletalMeshAsset;
+
+	UPROPERTY(EditAnywhere, BlueprintReadWrite, Category = "Settings", meta = (TitleProperty = "PhysicsAssetOverride"))
+	TObjectPtr<UPhysicsAsset> PhysicsAssetOverride;
+
+	/** Bodies to be checked against contact bodies */
+	UPROPERTY(EditAnywhere, BlueprintReadWrite, Category = "Settings", meta = (TitleProperty = "DomainBodyNames"))
+	TArray<FName> DomainBodyNames;
+	
+	UPROPERTY(EditAnywhere, BlueprintReadWrite, Category = "Settings", meta = (TitleProperty = "ContactBodyNames"))
+	TArray<FName> ContactBodyNames;
+
+	// Relative Body Notify subclass to create
+	UPROPERTY(EditAnywhere, BlueprintReadWrite, Category = "Settings", meta = (TitleProperty = "PhysicalIKRetargeterNotifyClass"))
+	TSubclassOf<URelativeBodyAnimNotifyBase> NotifyClass = nullptr;
+};
+
+USTRUCT(BlueprintType)
+struct FPropMeshSettings
+{
+	GENERATED_BODY()
+	
+	UPROPERTY(EditAnywhere, BlueprintReadWrite, Category = "Settings")
+	TObjectPtr<UStaticMesh> StaticMesh;
+
+	UPROPERTY(EditAnywhere, BlueprintReadWrite, Category = "Settings")
+	FName AttachBoneName;
+	
+	UPROPERTY(EditAnywhere, BlueprintReadWrite, Category = "Settings")
+	FTransform AttachTransform;
+};
+
+USTRUCT(BlueprintType)
+struct RELATIVEBODYANIMUTILS_API FPropBodyAnimModifierOptions : public FRelativeBodyAnimModifierOptions
+{
+	GENERATED_BODY()
+
+	// Settings for eash prop static mesh attachment to bake
+	UPROPERTY(EditAnywhere, BlueprintReadWrite, Category = "Settings")
+	TArray<FPropMeshSettings> PropMeshSettings;
+	
+	// Prop Notify subclass to create
+	UPROPERTY(EditAnywhere, BlueprintReadWrite, Category = "Settings")
+	TSubclassOf<URelativePropsBakeAnimNotify> PropNotifyClass = nullptr;
+};
+
+
+/**
+ * BPFL helper for setting up anim modifier 
+ */
+UCLASS(Transient, meta=(ScriptName="RelativeBodyAnimModifyHelper"))
+class RELATIVEBODYANIMUTILS_API URelativeBodyAnimBlueprintFunctions : public UBlueprintFunctionLibrary
+{
+	GENERATED_BODY()
+public:
+	
+	UFUNCTION(BlueprintCallable, Category = "Anim Modifier")
+	static void BulkUpdateRelativeBodyModifiers(UPARAM(ref) TArray<UAnimSequenceBase*>& AnimationList, const FRelativeBodyAnimModifierOptions& Options);
+	
+	UFUNCTION(BlueprintCallable, Category = "Anim Modifier")
+	static void BulkUpdatePropBodyModifiers(UPARAM(ref) TArray<UAnimSequenceBase*>& AnimationList, const FPropBodyAnimModifierOptions& Options);
+};
diff --git a/Engine/Plugins/Experimental/Animation/RelativeIKOp/Source/RelativeBodyAnimUtils/Public/RelativePropsAnimModifier.h b/Engine/Plugins/Experimental/Animation/RelativeIKOp/Source/RelativeBodyAnimUtils/Public/RelativePropsAnimModifier.h
--- a/Engine/Plugins/Experimental/Animation/RelativeIKOp/Source/RelativeBodyAnimUtils/Public/RelativePropsAnimModifier.h
+++ b/Engine/Plugins/Experimental/Animation/RelativeIKOp/Source/RelativeBodyAnimUtils/Public/RelativePropsAnimModifier.h
@@ -9,7 +9,7 @@
 #include "RelativePropsAnimModifier.generated.h"
 
 class UPhysicsAsset;
-class URelativeBodyAnimNotifyBase;
+class URelativePropsBakeAnimNotify;
 class USkeletalMesh;
 struct FAnimPose;
 struct FReferenceSkeleton;
@@ -41,7 +41,7 @@ public:
 
 	/**  Relative Props Notify subclass to create */
 	UPROPERTY(EditAnywhere, Category = "Settings", meta = (TitleProperty = "PropsNotifyClass"))
-	TSubclassOf<URelativeBodyAnimNotifyBase> PropsNotifyClass = nullptr;
+	TSubclassOf<URelativePropsBakeAnimNotify> PropsNotifyClass = nullptr;
 
 	UPROPERTY(EditAnywhere, Category = "Settings", meta = (TitleProperty = "PropsFloorInfoBaking"))
 	bool bPropsFloorInfoBaking = false;
