SpatiallyAwareOps: Rename RelativeIKOp plugin Friendly name to "Spatially Aware Retarget Ops".

#Synced-CL 48450904

diff --git a/Engine/Plugins/Experimental/Animation/RelativeIKOp/RelativeIKOp.uplugin b/Engine/Plugins/Experimental/Animation/RelativeIKOp/RelativeIKOp.uplugin
--- a/Engine/Plugins/Experimental/Animation/RelativeIKOp/RelativeIKOp.uplugin
+++ b/Engine/Plugins/Experimental/Animation/RelativeIKOp/RelativeIKOp.uplugin
@@ -2,8 +2,8 @@
 	"FileVersion": 3,
 	"Version": 1,
 	"VersionName": "1.0",
-	"FriendlyName": "Relative IK Op",
-	"Description": "Experimental Retarge Op for Relative IK goal setting",
+	"FriendlyName": "Spatially Aware Retarget Ops",
+	"Description": "A collection of Retarget Ops for preserving spatial relationships on retargeted animations",
 	"Category": "Animation",
 	"CreatedBy": "Epic Games, Inc.",
 	"CreatedByURL": "https://epicgames.com",
