Relative IK:
+ Floor Constraints

#Synced-CL 48442684

diff --git a/Engine/Plugins/Experimental/Animation/RelativeIKOp/Source/RelativeBodyAnimUtils/Private/RelativePropsAnimModifier.cpp b/Engine/Plugins/Experimental/Animation/RelativeIKOp/Source/RelativeBodyAnimUtils/Private/RelativePropsAnimModifier.cpp
--- a/Engine/Plugins/Experimental/Animation/RelativeIKOp/Source/RelativeBodyAnimUtils/Private/RelativePropsAnimModifier.cpp
+++ b/Engine/Plugins/Experimental/Animation/RelativeIKOp/Source/RelativeBodyAnimUtils/Private/RelativePropsAnimModifier.cpp
@@ -319,12 +319,13 @@ void URelativePropsAnimModifier::OnApply_Implementation(UAnimSequence* InAnimati
 					}
 					if (bPropsFloorInfoBaking)
 					{
+						FTransform3f PropOffsetTfm(BakedAnimNotifyInfo->OffsetTransformsForBones[PropsAttachBones[PropsIdx]]);
 						int32 Offset = SampleIndex * BakedAnimNotifyInfo->PropsPairs.Num() + Count * 2;
 						int32 LowestVertexIndex = FindLowestIndex(PropsVLocations[PropsIdx]);
 						int32 HighestVertexIndex=FindHighestIndex(PropsVLocations[PropsIdx]);
 
-						BakedAnimNotifyInfo->PropsPairsLocalReference[Offset + 0] = PositionVertexBuffer.VertexPosition(LowestVertexIndex);
-						BakedAnimNotifyInfo->PropsPairsLocalReference[Offset + 1] = PositionVertexBuffer.VertexPosition(HighestVertexIndex);
+						BakedAnimNotifyInfo->PropsPairsLocalReference[Offset + 0] = PropOffsetTfm.TransformPosition(PositionVertexBuffer.VertexPosition(LowestVertexIndex));
+						BakedAnimNotifyInfo->PropsPairsLocalReference[Offset + 1] = PropOffsetTfm.TransformPosition(PositionVertexBuffer.VertexPosition(HighestVertexIndex));
 
 						Count++;
 					}
diff --git a/Engine/Plugins/Experimental/Animation/RelativeIKOp/Source/RelativeIKOp/Private/RelativeIKOp.cpp b/Engine/Plugins/Experimental/Animation/RelativeIKOp/Source/RelativeIKOp/Private/RelativeIKOp.cpp
--- a/Engine/Plugins/Experimental/Animation/RelativeIKOp/Source/RelativeIKOp/Private/RelativeIKOp.cpp
+++ b/Engine/Plugins/Experimental/Animation/RelativeIKOp/Source/RelativeIKOp/Private/RelativeIKOp.cpp
@@ -254,13 +254,8 @@ void FIKRetargetRelativeIKOp::ComputeFramePairSpaces(
 	{
 		// TODO: Allow selecting a subset of body pairs in retargeter
 		const int32 OutPairIdx = PairIdx + PairOffset;
-		FrameBoneWeights[2*OutPairIdx] = 0;
-		FrameBoneWeights[2*OutPairIdx + 1] = 0;
-		FrameBoneOffsets[2*OutPairIdx] = FVector::ZeroVector;
-		FrameBoneOffsets[2*OutPairIdx + 1] = FVector::ZeroVector;
-		//TODO: Remove FrameBonePositions and figure out a quick access from effect vert indices to bone positions
-		FrameBonePositions[2*OutPairIdx] = FVector::ZeroVector;
-		FrameBonePositions[2*OutPairIdx + 1] = FVector::ZeroVector;
+		FrameBoneConstraints[2*OutPairIdx].Reset();
+		FrameBoneConstraints[2*OutPairIdx+1].Reset();
 		
 		FName SourcePrimaryBone = CachedNotifyInfo->BodyPairs[2*PairIdx];
 		FName SourceOtherBone = CachedNotifyInfo->BodyPairs[2*PairIdx + 1];
@@ -348,7 +343,7 @@ void FIKRetargetRelativeIKOp::ComputeFramePairSpaces(
 
 		FVector TargetPrimaryBoneLoc = TargetPose[TargetPrimaryPoseIdx].GetLocation();
 		FVector TargetOtherBoneLoc = TargetPose[TargetOtherPoseIdx].GetLocation();
-		
+	
 		FDoubleInterval FeasibleRange_O(0.0, 1.0);
 		FDoubleInterval FeasibleRange_P(0.0, 1.0);
 		UpdateCacheChainInfo(TargetPose, TargetOtherBone);
@@ -361,32 +356,27 @@ void FIKRetargetRelativeIKOp::ComputeFramePairSpaces(
 
 		double ContactBlendWeight = (Settings.bBidirectionalDistanceContactBlend) ? LinearWeight : 0.5;
 		double FeasibilityWeight_O = FMath::Lerp(FeasibleRange_O.Min, FeasibleRange_O.Max, ContactBlendWeight);
-		FVector WeightedTargetOffset_O = TargetOffsetVec * FeasibilityWeight_O;
-
 		double FeasibilityWeight_P = FMath::Lerp(FeasibleRange_P.Min, FeasibleRange_P.Max, ContactBlendWeight);
-		FVector WeightedTargetOffset_P = -TargetOffsetVec * FeasibilityWeight_P;
+		FVector WeightedTargetOffset_O = TargetOffsetVec * (Settings.bMultiBoneSolve? 0.5*(FeasibilityWeight_O+FeasibilityWeight_P): FeasibilityWeight_O);
+		FVector WeightedTargetOffset_P = -TargetOffsetVec * (Settings.bMultiBoneSolve? 0.5*(FeasibilityWeight_O+FeasibilityWeight_P): FeasibilityWeight_P);
 		
-		FrameBoneOffsets[2*OutPairIdx] = WeightedTargetOffset_P;
-		FrameBoneOffsets[2*OutPairIdx + 1] = WeightedTargetOffset_O;
-
-		FrameBonePositions[2*OutPairIdx] = TargetPrimaryBoneLoc;
-		FrameBonePositions[2*OutPairIdx + 1] = TargetOtherBoneLoc;
-
-		FrameBoneWeights[2*OutPairIdx] = Weight;
-		FrameBoneWeights[2*OutPairIdx + 1] = Weight;
+		FrameBoneConstraints[2*OutPairIdx].Offset = WeightedTargetOffset_P;
+		FrameBoneConstraints[2*OutPairIdx+1].Offset = WeightedTargetOffset_O;
+		FrameBoneConstraints[2*OutPairIdx].Weight = Weight;
+		FrameBoneConstraints[2*OutPairIdx+1].Weight = Weight;
 
 #if WITH_EDITOR
 		// Debug draw verts
 		if (Settings.bDebugDraw)
 		{
-			DebugInfo.SourcePairVerts.Add({SourcePrimaryCompV, SourceOtherCompV, false, FVector::ZeroVector, FVector::ZeroVector});
+			DebugInfo.SourcePairVerts.Add({SourcePrimaryCompV, SourceOtherCompV, false, FVector::ZeroVector, FVector::ZeroVector, Weight});
 			if (Settings.DebugFullRetargetPairBones.Contains(TargetPrimaryBone))
 			{
-				DebugInfo.TargetPairVerts.Add({TargetPrimaryCompV_P, TargetOtherCompV_O, true,TargetOtherCompV_P, TargetPrimaryCompV_O,});
+				DebugInfo.TargetPairVerts.Add({TargetPrimaryCompV_P, TargetOtherCompV_O, true,TargetOtherCompV_P, TargetPrimaryCompV_O, Weight});
 			}
 			else
 			{
-				DebugInfo.TargetPairVerts.Add({TargetPrimaryCompV_P, TargetOtherCompV_O, false, FVector::ZeroVector, FVector::ZeroVector});
+				DebugInfo.TargetPairVerts.Add({TargetPrimaryCompV_P, TargetOtherCompV_O, false, FVector::ZeroVector, FVector::ZeroVector, Weight});
 			}
 			
 			DebugPairSpaces[2*OutPairIdx] = FDebugRelativeTargetPairSpace{TargetPrimaryCompV_P, TargetPrimaryCompV_P-TargetOffsetVec, FeasibilityWeight_P, FeasibleRange_P};
@@ -417,34 +407,86 @@ void FIKRetargetRelativeIKOp::ComputeFramePropSpaces(
 	const double RetargetContactAlpha = FMath::Clamp(Settings.RetargetContactAlpha, 0.0, 1.0);
 	const double RetargetSpringAlpha = FMath::Clamp(Settings.RetargetSpringAlpha, 0.0, 1.0);
 	// Get pair target body verts and weights
-	// Get pair target body verts and weights
 	for (int PairIdx = 0; PairIdx < PairCount; ++PairIdx)
 	{
 		const int32 OutPairIdx = PairIdx + PairOffset;
-		FrameBoneWeights[2*OutPairIdx] = 0;
-		FrameBoneWeights[2*OutPairIdx + 1] = 0;
-		FrameBoneOffsets[2*OutPairIdx] = FVector::ZeroVector;
-		FrameBoneOffsets[2*OutPairIdx + 1] = FVector::ZeroVector;
-		//TODO: Remove FrameBonePositions and figure out a quick access from effect vert indices to bone positions
-		FrameBonePositions[2*OutPairIdx] = FVector::ZeroVector;
-		FrameBonePositions[2*OutPairIdx + 1] = FVector::ZeroVector;
-
+		FrameBoneConstraints[2*OutPairIdx].Reset();
+		FrameBoneConstraints[2*OutPairIdx+1].Reset();
+		
 		FName SourcePropBone = CachedPropNotifyInfo->PropsPairs[2 * PairIdx];
 		FName TargetPropBone = ApplyBodyMap(SourcePropBone);
 		FName SourceBodyBone = CachedPropNotifyInfo->PropsPairs[2*PairIdx + 1];
 		FName TargetBodyBone = ApplyBodyMap(SourceBodyBone);
 
-		if (TargetPropBone == NAME_None || TargetBodyBone == NAME_None)
+		if (TargetPropBone == NAME_None || !CacheSourceSkelIndices.Contains(SourcePropBone) || !CacheTargetSkelIndices.Contains(TargetPropBone) )
 		{
 			continue;
 		}
 
-		if (!CacheSourceSkelIndices.Contains(SourcePropBone) || !CacheSourceSkelIndices.Contains(SourceBodyBone))
+		if (SourceBodyBone==FName("Floor"))
 		{
+			int32 PinPropIdx = CachePinSourcePropBones.Find(SourcePropBone);
+			if (PinPropIdx == INDEX_NONE)
+			{
+				return;
+			}
+			
+			int32 SourcePropPoseIdx = CacheSourceSkelIndices[SourcePropBone];
+			int32 TargetPropPoseIdx = CacheTargetSkelIndices[TargetPropBone];
+			
+			FTransform SourcePropTfm = SourceScaleTfm * SourcePose[SourcePropPoseIdx];
+			FTransform TargetPropBoneDelta = GetTargetBoneDelta(TargetPropBone);
+			FTransform TargetPropTfm = TargetPropBoneDelta * TargetPose[TargetPropPoseIdx];
+			
+			FVector SourcePropBakeZMin(FramePointView[2*PairIdx]);
+			FVector SourcePropBakeZMax(FramePointView[2*PairIdx+1]);
+
+			FVector SourcePropCompZMin = SourcePropTfm.TransformPosition(SourcePropBakeZMin);
+			FVector SourcePropCompZMax = SourcePropTfm.TransformPosition(SourcePropBakeZMax);
+
+			FVector TargetPropCompZMin = TargetPropTfm.TransformPosition(SourcePropBakeZMin);
+			FVector TargetPropCompZMax = TargetPropTfm.TransformPosition(SourcePropBakeZMax);
+			
+			// Rescale components for distance checks if "Ignore sources scale" is on
+			// NOTE: This works because we don't need absolute coordinates (just distance).
+			//       If absolute coords needed then must properly account for scale pivot
+			FVector RescaleSourcePropZMin = InvSourceScale * SourcePropCompZMin;
+			FVector RescaleSourcePropZMax = InvSourceScale * SourcePropCompZMax;
+			double SourceZDist = RescaleSourcePropZMin.Z;
+			double SourceWeight = GetDistanceWeight(Settings.DistanceWeightMode, SourceZDist, Settings.FloorThreshold, Settings.DistHalfPointLambda);
+			double TargetZDist = TargetPropCompZMin.Z;
+			double TargetWeight = GetDistanceWeight(Settings.DistanceWeightMode, TargetZDist, Settings.FloorThreshold, Settings.DistHalfPointLambda);
+			
+			// Offset vector (Prop -> Floor) to Lerp for dist alph
+			double FloorOffsetZMin = Settings.DistHalfPointLambda*Settings.FloorThreshold-TargetPropCompZMin.Z;
+			FVector TargetOffsetZMin = SourcePropCompZMin - TargetPropCompZMin;
+			FVector TargetOffsetZMax = SourcePropCompZMax - TargetPropCompZMax;
+			double Weight = SourceWeight*SourceWeight + (1-SourceWeight)*TargetWeight;
+
+			FrameBoneConstraints[2*OutPairIdx].ConstraintType = ERelativeIKFrameBoneConstraintType::Z;
+			FrameBoneConstraints[2*OutPairIdx+1].ConstraintType = ERelativeIKFrameBoneConstraintType::Z;
+			FrameBoneConstraints[2*OutPairIdx].Offset = {0,0,SourceWeight*TargetOffsetZMin.Z + (1-SourceWeight)*FloorOffsetZMin};
+			FrameBoneConstraints[2*OutPairIdx+1].Offset = {0,0,SourceWeight*TargetOffsetZMax.Z + (1-SourceWeight)*FloorOffsetZMin};
+			FrameBoneConstraints[2*OutPairIdx].Weight = Weight;
+			FrameBoneConstraints[2*OutPairIdx+1].Weight = Weight;
+
+#if WITH_EDITOR
+			// Debug draw verts
+			if (Settings.bDebugDraw)
+			{
+				FVector SourceFloorZMin = SourcePropCompZMin;
+				SourceFloorZMin.Z = 0;
+				FVector TargetFloorZMin = TargetPropCompZMin;
+				TargetFloorZMin.Z = 0;
+				
+				DebugInfo.SourcePairVerts.Add({SourcePropCompZMin, SourceFloorZMin, false, FVector::ZeroVector, FVector::ZeroVector, Weight});
+				DebugInfo.TargetPairVerts.Add({TargetPropCompZMin, TargetFloorZMin, false, FVector::ZeroVector, FVector::ZeroVector, Weight});
+			}
+#endif //WITH_EDITOR
 			continue;
 		}
-
-		if (!CacheTargetSkelIndices.Contains(TargetPropBone) || !CacheTargetSkelIndices.Contains(TargetBodyBone))
+		
+		if (TargetBodyBone == NAME_None || !CacheSourceSkelIndices.Contains(SourceBodyBone) || !CacheTargetSkelIndices.Contains(TargetBodyBone))
 		{
 			continue;
 		}
@@ -504,8 +546,7 @@ void FIKRetargetRelativeIKOp::ComputeFramePropSpaces(
 		// Compute body vert represented in component space relative to prop (_P) and body (_B)
 		FVector TargetBodyCompV_B = ApplyBodyTransform(TargetBodyTfm, SourceBoneBakeBodyV);
 		FVector TargetBodyCompV_P = TargetPropTfm.TransformPosition(SourceBodyV_P);
-
-
+		
 		// These are the target midpoint positions used for spring constraints in PhysicalIK
 		FVector TargetSpring_P = FMath::Lerp(TargetPropCompV_P,TargetBodyCompV_P, RetargetSpringAlpha);
 		FVector TargetSpring_B = FMath::Lerp(TargetPropCompV_B, TargetBodyCompV_B, RetargetSpringAlpha);
@@ -515,7 +556,7 @@ void FIKRetargetRelativeIKOp::ComputeFramePropSpaces(
 
 		FVector TargetPropBoneLoc = TargetPose[TargetPropPoseIdx].GetLocation();
 		FVector TargetBodyBoneLoc = TargetPose[TargetBodyPoseIdx].GetLocation();
-
+		
 		FDoubleInterval FeasibleRange_P(0.0, 1.0);
 		FDoubleInterval FeasibleRange_B(0.0, 1.0);
 		UpdateCacheChainInfo(TargetPose, TargetBodyBone);
@@ -526,32 +567,28 @@ void FIKRetargetRelativeIKOp::ComputeFramePropSpaces(
 
 		double ContactBlendWeight = (Settings.bBidirectionalDistanceContactBlend) ? LinearWeight : 0.5;
 		double FeasibilityWeight_B =  FMath::Lerp(FeasibleRange_B.Min, FeasibleRange_B.Max, ContactBlendWeight);
-		FVector WeightedTargetOffset_B = TargetOffsetVec * FeasibilityWeight_B;
-
 		double FeasibilityWeight_P = FMath::Lerp(FeasibleRange_P.Min, FeasibleRange_P.Max, ContactBlendWeight);
-		FVector WeightedTargetOffset_P = -TargetOffsetVec * FeasibilityWeight_P;
-
-		FrameBoneOffsets[2*OutPairIdx] = WeightedTargetOffset_P;
-		FrameBoneOffsets[2*OutPairIdx + 1] = WeightedTargetOffset_B;
-
-		FrameBonePositions[2*OutPairIdx] = TargetPropBoneLoc;
-		FrameBonePositions[2*OutPairIdx + 1] = TargetBodyBoneLoc;
+		
+		FVector WeightedTargetOffset_B = TargetOffsetVec * (Settings.bMultiBoneSolve? 0.5*(FeasibilityWeight_B+FeasibilityWeight_P): FeasibilityWeight_B);
+		FVector WeightedTargetOffset_P = -TargetOffsetVec * (Settings.bMultiBoneSolve? 0.5*(FeasibilityWeight_B+FeasibilityWeight_P): FeasibilityWeight_P);
 
-		FrameBoneWeights[2*OutPairIdx] = Weight;
-		FrameBoneWeights[2*OutPairIdx + 1] = Weight;
+		FrameBoneConstraints[2*OutPairIdx].Offset = WeightedTargetOffset_P;
+		FrameBoneConstraints[2*OutPairIdx+1].Offset = WeightedTargetOffset_B;
+		FrameBoneConstraints[2*OutPairIdx].Weight = Weight;
+		FrameBoneConstraints[2*OutPairIdx+1].Weight = Weight;
 
 #if WITH_EDITOR
 		// Debug draw verts
 		if (Settings.bDebugDraw)
 		{
-			DebugInfo.SourcePairVerts.Add({SourcePropCompV, SourceBodyCompV, false, FVector::ZeroVector, FVector::ZeroVector});
+			DebugInfo.SourcePairVerts.Add({SourcePropCompV, SourceBodyCompV, false, FVector::ZeroVector, FVector::ZeroVector, Weight});
 			if (Settings.DebugFullRetargetPairBones.Contains(TargetPropBone))
 			{
-				DebugInfo.TargetPairVerts.Add({TargetPropCompV_P, TargetBodyCompV_B, true,TargetBodyCompV_P, TargetPropCompV_B,});
+				DebugInfo.TargetPairVerts.Add({TargetPropCompV_P, TargetBodyCompV_B, true,TargetBodyCompV_P, TargetPropCompV_B, Weight});
 			}
 			else
 			{
-				DebugInfo.TargetPairVerts.Add({TargetPropCompV_P, TargetBodyCompV_B, false, FVector::ZeroVector, FVector::ZeroVector});
+				DebugInfo.TargetPairVerts.Add({TargetPropCompV_P, TargetBodyCompV_B, false, FVector::ZeroVector, FVector::ZeroVector, Weight});
 			}
 			
 			DebugPairSpaces[2*OutPairIdx] = FDebugRelativeTargetPairSpace{TargetPropCompV_P, TargetPropCompV_P-TargetOffsetVec, FeasibilityWeight_P, FeasibleRange_P};
@@ -802,15 +839,14 @@ void FIKRetargetRelativeIKOp::ApplyAverageBodyPairTargets(
 		FVector BoneLoc = TargetBonePose.GetLocation();
 
 		FVector TargetOffset;
-		double TotalWeight = ComputeTargetWeightedOffset(TargetOffset, EffectIndices, BoneLoc, DebugInfo, DebugPairSpaces);
-		if (FMath::IsNearlyZero(TotalWeight))
+		FVector TotalWeight = ComputeTargetWeightedOffset(TargetOffset, EffectIndices, BoneLoc, DebugInfo, DebugPairSpaces);
+		if (FMath::IsNearlyZero(TotalWeight.GetMax()))
 		{
 			continue;
 		}
-
-		FVector TargetLoc = BoneLoc + TargetOffset;
-		UpdateTargetIKGoal(OutIKGoals, TargetBoneName, TotalWeight, TargetLoc, TargetBonePose);
-		UpdatePropPinTransform(SourceBoneName, TotalWeight, TargetLoc, TargetBonePose);
+		
+		UpdateTargetIKGoal(OutIKGoals, TargetBoneName, TotalWeight, TargetOffset, TargetBonePose);
+		UpdatePropPinTransform(SourceBoneName, TotalWeight, TargetOffset, TargetBonePose);
 	}
 }
 
@@ -850,14 +886,13 @@ void FIKRetargetRelativeIKOp::ApplyMatrixBodyPairTargets(
 		OutBoneNames.Add(SourceBoneName);
 		OutBoneLocations.Add(TargetBonePose.GetLocation());
 	}
+	
+	Eigen::MatrixXd BaseWeightMatrix = Eigen::MatrixXd::Identity(OutIndexMap.Num(), OutIndexMap.Num());
+	Eigen::MatrixXd OffsetMatrix = Eigen::MatrixXd::Zero(OutIndexMap.Num(), 3);
+	Eigen::MatrixXd DeltaWeights = Eigen::MatrixXd::Zero(OutIndexMap.Num(), 3);
+	Eigen::VectorXd BaseWeights = Eigen::VectorXd::Zero(OutIndexMap.Num());
 
-	Eigen::MatrixXd WeightMatrix(OutIndexMap.Num(), OutIndexMap.Num());
-	Eigen::MatrixXd OffsetMatrix(OutIndexMap.Num(), 3);
-	Eigen::VectorXd TotalWeights(OutIndexMap.Num());
-
-	WeightMatrix = Eigen::MatrixXd::Identity(OutIndexMap.Num(), OutIndexMap.Num());
-	OffsetMatrix = Eigen::MatrixXd::Zero(OutIndexMap.Num(), 3);
-	TotalWeights = Eigen::VectorXd::Zero(OutIndexMap.Num());
+	double MaxDelta = 0.0;
 	for (int32 RowIdx = 0; RowIdx < OutIndexMap.Num(); ++RowIdx)
 	{
 		int32 EffectIdx = OutIndexMap[RowIdx];
@@ -867,23 +902,48 @@ void FIKRetargetRelativeIKOp::ApplyMatrixBodyPairTargets(
 			continue;
 		}
 		
-		double TotalWeight = ComputeTargetWeightedOffsetRow(OffsetMatrix, RowIdx, OutBoneNames, EffectIndices, OutBoneLocations[RowIdx], DebugInfo, DebugPairSpaces);
-		if (FMath::IsNearlyZero(TotalWeight))
+		FVector TotalWeight = ComputeTargetWeightedOffsetRow(OffsetMatrix, RowIdx, OutBoneNames, EffectIndices, OutBoneLocations[RowIdx], DebugInfo, DebugPairSpaces);
+		if (FMath::IsNearlyZero(TotalWeight.GetMax()))
 		{
 			continue;
 		}
 		
-		ComputeWeightMatrixRow(WeightMatrix, RowIdx, OutBoneNames, EffectIndices);
-		TotalWeights(RowIdx) = TotalWeight;
-		WeightMatrix(RowIdx, RowIdx) = TotalWeight + Settings.MultiBoneRegularization;
+		BaseWeights(RowIdx) = TotalWeight.GetMin();
+		DeltaWeights(RowIdx,0) = TotalWeight.X - BaseWeights(RowIdx);
+		DeltaWeights(RowIdx,1) = TotalWeight.Y - BaseWeights(RowIdx);
+		DeltaWeights(RowIdx,2) = TotalWeight.Z - BaseWeights(RowIdx);
+		
+		ComputeWeightMatrixRow(BaseWeightMatrix, RowIdx, OutBoneNames, EffectIndices);
+		BaseWeightMatrix(RowIdx, RowIdx) = BaseWeights(RowIdx) + Settings.MultiBoneRegularization;
 	}
-
-	Eigen::MatrixXd SolvedOffset = WeightMatrix.jacobiSvd(Eigen::ComputeThinU | Eigen::ComputeThinV).solve(OffsetMatrix);
-	// Eigen::MatrixXd SolvedOffset = WeightMatrix.ldlt().solve(OffsetMatrix);
+	
+	Eigen::MatrixXd SolvedOffset(OutIndexMap.Num(), 3);
+	if (FMath::IsNearlyZero(DeltaWeights.maxCoeff()))
+	{
+		// Eigen::ColPivHouseholderQR<Eigen::MatrixXd> qr = BaseWeightMatrix.colPivHouseholderQr();
+		// Eigen::LLT<Eigen::MatrixXd> MatrixSolver = BaseWeightMatrix.llt();
+		Eigen::JacobiSVD<Eigen::MatrixXd> MatrixSolver = BaseWeightMatrix.jacobiSvd(Eigen::ComputeThinU | Eigen::ComputeThinV);
+		SolvedOffset = MatrixSolver.solve(OffsetMatrix);
+	}
+	else
+	{
+		// TODO: Generalize for partial conditions on any/all coordinate sets
+		// Eigen::LLT<Eigen::MatrixXd> MatrixSolver = BaseWeightMatrix.llt();
+		Eigen::JacobiSVD<Eigen::MatrixXd> MatrixSolver = BaseWeightMatrix.jacobiSvd(Eigen::ComputeThinU | Eigen::ComputeThinV);
+		SolvedOffset.leftCols<2>() = MatrixSolver.solve(OffsetMatrix.leftCols<2>());
+		
+		//TODO: This is kinda a nasty eigen op, could alternatively build this up in since update matrix by add/solve/subtract
+		Eigen::MatrixXd UpdatedWeights = BaseWeightMatrix + DeltaWeights.col(2).asDiagonal().toDenseMatrix();
+		// MatrixSolver = UpdatedWeights.llt();
+		MatrixSolver = UpdatedWeights.jacobiSvd(Eigen::ComputeThinU | Eigen::ComputeThinV);
+		SolvedOffset.rightCols<1>() = MatrixSolver.solve(OffsetMatrix.rightCols<1>());
+	}
+	
 	for (int32 RowIdx = 0; RowIdx < OutIndexMap.Num(); ++RowIdx)
 	{
 		int32 EffectIdx = OutIndexMap[RowIdx];
-		if (FMath::IsNearlyZero(TotalWeights(RowIdx)))
+		FVector TotalWeight = FVector{DeltaWeights(RowIdx, 0), DeltaWeights(RowIdx, 1), DeltaWeights(RowIdx, 2)} + BaseWeights(RowIdx);
+		if (FMath::IsNearlyZero(TotalWeight.GetMax()))
 		{
 			continue;
 		}
@@ -895,22 +955,25 @@ void FIKRetargetRelativeIKOp::ApplyMatrixBodyPairTargets(
 			continue;
 		}
 		
-		FVector TargetOffset{SolvedOffset(RowIdx, 0), SolvedOffset(RowIdx, 1), SolvedOffset(RowIdx, 2)};
+		FVector TargetOffset{
+			(!FMath::IsNearlyZero(TotalWeight.X)) ? SolvedOffset(RowIdx, 0) : 0.0,
+			(!FMath::IsNearlyZero(TotalWeight.Y)) ? SolvedOffset(RowIdx, 1) : 0.0,
+			(!FMath::IsNearlyZero(TotalWeight.Z)) ? SolvedOffset(RowIdx, 2) : 0.0
+		};
 
 		int32 TargetBoneIdx = CacheTargetSkelIndices[TargetBoneName];
 		const FTransform& TargetBonePose = TargetPose[TargetBoneIdx];
-
-		FVector TargetLoc = TargetBonePose.GetLocation() + TargetOffset;
-		UpdateTargetIKGoal(OutIKGoals, TargetBoneName, TotalWeights(RowIdx), TargetLoc, TargetBonePose);
-		UpdatePropPinTransform(SourceBoneName, TotalWeights(RowIdx), TargetLoc, TargetBonePose);
+		
+		UpdateTargetIKGoal(OutIKGoals, TargetBoneName, TotalWeight, TargetOffset, TargetBonePose);
+		UpdatePropPinTransform(SourceBoneName, TotalWeight, TargetOffset, TargetBonePose);
 	}
 }
 
 void FIKRetargetRelativeIKOp::UpdateTargetIKGoal(
 	TArray<FIKRigGoal>& OutIKGoals,
 	FName TargetBoneName,
-	double TargetWeight,
-	const FVector& TargetGoalLoc,
+	const FVector& TargetWeight,
+	const FVector& TargetGoalOffset,
 	const FTransform& InTargetBonePose)
 {
 	FIKRigGoal* IKGoal = OutIKGoals.FindByPredicate([TargetBoneName](const FIKRigGoal& Other){return Other.BoneName == TargetBoneName;});
@@ -920,7 +983,8 @@ void FIKRetargetRelativeIKOp::UpdateTargetIKGoal(
 	}
 		
 	// TODO: Play with total weight functions to alpha goal
-	double GoalAlpha = FMath::Min(1.0, TargetWeight*FMath::Clamp(Settings.ContributionSumWeight, 0.0f, 1.0f));
+	FVector TargetGoalLoc = InTargetBonePose.GetLocation() + TargetGoalOffset;
+	FVector GoalAlpha = FVector::Min(FVector::OneVector, TargetWeight*FMath::Clamp(Settings.ContributionSumWeight, 0.0f, 1.0f));
 	FVector PrevGoal = GoalLocBlendCompSpace(IKGoal, InTargetBonePose);
 	FVector BlendGoal = FMath::Lerp(PrevGoal, TargetGoalLoc, GoalAlpha);
 
@@ -938,8 +1002,8 @@ void FIKRetargetRelativeIKOp::UpdateTargetIKGoal(
 
 void FIKRetargetRelativeIKOp::UpdatePropPinTransform(
 	FName SourceBoneName,
-	double TargetWeight,
-	const FVector& TargetPropLoc,
+	const FVector& TargetWeight,
+	const FVector& TargetPropOffset,
 	const FTransform& InTargetBonePose)
 {
 	int32 PinPropIdx = CachePinSourcePropBones.Find(SourceBoneName);
@@ -948,21 +1012,18 @@ void FIKRetargetRelativeIKOp::UpdatePropPinTransform(
 		return;
 	}
 	
-	double PropBoneAlpha = FMath::Min(1.0, TargetWeight*FMath::Clamp(Settings.ContributionSumWeight, 0.0, 1.0));
-	double BlendAlpha = PropBoneAlpha * (1.0-Settings.BlendPinsToSource);
-	FVector BlendBoneLoc = FMath::Lerp(InTargetBonePose.GetLocation(), TargetPropLoc, BlendAlpha);
-	
-	CachePinTargetPropTransforms[PinPropIdx].SetLocation(BlendBoneLoc);
+	FVector BlendAlpha = FVector::Min(FVector::OneVector, TargetWeight*FMath::Clamp(Settings.ContributionSumWeight, 0.0, 1.0)) * (1.0-Settings.BlendPinsToSource);
+	CachePinTargetPropTransforms[PinPropIdx].SetLocation(InTargetBonePose.GetLocation() + BlendAlpha*TargetPropOffset);
 }
 
-double FIKRetargetRelativeIKOp::ComputeTargetWeightedOffset(
+FVector FIKRetargetRelativeIKOp::ComputeTargetWeightedOffset(
 	FVector& OutTargetOffset,
 	const TArray<int32>& EffectIndices,
 	const FVector& BoneLoc,
 	FDebugRelativeIKDrawInfo& DebugInfo,
 	TArray<FDebugRelativeTargetPairSpace>& DebugPairSpaces)
 {
-	double TotalWeight = 0;
+	FVector TotalWeight = FVector::ZeroVector;
 	OutTargetOffset = FVector::ZeroVector;
 
 #if WITH_EDITOR
@@ -975,30 +1036,41 @@ double FIKRetargetRelativeIKOp::ComputeTargetWeightedOffset(
 	
 	for (int Index : EffectIndices)
 	{
-		if (FMath::IsNearlyZero(FrameBoneWeights[Index]))
+		const FRelativeIKFrameBoneConstraint& BoneConstraint = FrameBoneConstraints[Index];
+		if (FMath::IsNearlyZero(BoneConstraint.Weight))
 		{
 			continue;
 		}
-		
-		OutTargetOffset += FrameBoneWeights[Index] * FrameBoneOffsets[Index];
-		TotalWeight += FrameBoneWeights[Index];
+
+		FVector TargetOffset = BoneConstraint.Offset;
+		FVector TargetWeight = FVector(BoneConstraint.Weight);
+		if (BoneConstraint.ConstraintType == ERelativeIKFrameBoneConstraintType::Z)
+		{
+			TargetWeight = {0,0,BoneConstraint.Weight};
+		}
+
+		OutTargetOffset += TargetWeight * TargetOffset;
+		TotalWeight += TargetWeight;
 
 #if WITH_EDITOR
 		if (Settings.bDebugDraw)
 		{
-			FVector TargetLoc = FrameBoneOffsets[Index] + BoneLoc;
-			DebugInfo.TargetGoals.Last().PairTargets.Add(TargetLoc);
+			// TODO: Fixup target debug to show dynamic offset goals (post-solve update?)
+			FVector BoneTarget = TargetOffset + BoneLoc;
+			DebugInfo.TargetGoals.Last().PairTargets.Add(BoneTarget);
 			DebugInfo.PairRetargetInfo.Emplace(DebugPairSpaces[Index]);
 		}
 #endif //WITH_EDITOR
 	}
-		
-	if (FMath::IsNearlyZero(TotalWeight))
-	{
-		return TotalWeight;
-	}
-		
-	OutTargetOffset /= TotalWeight;
+	
+	FVector InvWeight = {
+		(FMath::IsNearlyZero(TotalWeight.X) ? 0.0 : 1.0 / TotalWeight.X),
+		(FMath::IsNearlyZero(TotalWeight.Y) ? 0.0 : 1.0 / TotalWeight.Y),
+		(FMath::IsNearlyZero(TotalWeight.Z) ? 0.0 : 1.0 / TotalWeight.Z),
+
+	};
+	OutTargetOffset *= InvWeight;
+	
 #if WITH_EDITOR
 	if (Settings.bDebugDraw)
 	{
@@ -1010,7 +1082,6 @@ double FIKRetargetRelativeIKOp::ComputeTargetWeightedOffset(
 }
 
 
-
 void FIKRetargetRelativeIKOp::ComputeWeightMatrixRow(
 	Eigen::MatrixXd& WeightMatrix,
 	int32 RowIdx,
@@ -1019,7 +1090,8 @@ void FIKRetargetRelativeIKOp::ComputeWeightMatrixRow(
 {
 	for (int Index : EffectIndices)
 	{
-		if (FMath::IsNearlyZero(FrameBoneWeights[Index]))
+		const FRelativeIKFrameBoneConstraint& BoneConstraint = FrameBoneConstraints[Index];
+		if (FMath::IsNearlyZero(BoneConstraint.Weight) || BoneConstraint.ConstraintType != ERelativeIKFrameBoneConstraintType::XYZ)
 		{
 			continue;
 		}
@@ -1031,11 +1103,11 @@ void FIKRetargetRelativeIKOp::ComputeWeightMatrixRow(
 			continue;
 		}
 
-		WeightMatrix(RowIdx, OutWeightIdx) = -FrameBoneWeights[Index];
+		WeightMatrix(RowIdx, OutWeightIdx) = -BoneConstraint.Weight;
 	}
 }
 
-double FIKRetargetRelativeIKOp::ComputeTargetWeightedOffsetRow(
+FVector FIKRetargetRelativeIKOp::ComputeTargetWeightedOffsetRow(
 	Eigen::MatrixXd& OffsetMatrix,
 	int32 RowIdx,
 	const TArray<FName>& OutBones,
@@ -1044,7 +1116,7 @@ double FIKRetargetRelativeIKOp::ComputeTargetWeightedOffsetRow(
 	FDebugRelativeIKDrawInfo& DebugInfo,
 	TArray<FDebugRelativeTargetPairSpace>& DebugPairSpaces)
 {
-	double TotalWeight = 0;
+	FVector TotalWeight = FVector::ZeroVector;
 	FVector OffsetLoc = FVector::ZeroVector;
 
 #if WITH_EDITOR
@@ -1059,23 +1131,29 @@ double FIKRetargetRelativeIKOp::ComputeTargetWeightedOffsetRow(
 	
 	for (int Index : EffectIndices)
 	{
-		if (FMath::IsNearlyZero(FrameBoneWeights[Index]))
+		const FRelativeIKFrameBoneConstraint& BoneConstraint = FrameBoneConstraints[Index];
+		if (FMath::IsNearlyZero(BoneConstraint.Weight))
 		{
 			continue;
 		}
 
-		FVector TargetOffset = FrameBoneOffsets[Index];
+		FVector TargetOffset = BoneConstraint.Offset;
+		FVector TargetWeight = FVector(BoneConstraint.Weight);
+		if (BoneConstraint.ConstraintType == ERelativeIKFrameBoneConstraintType::Z)
+		{
+			TargetWeight = {0,0,BoneConstraint.Weight};
+		}
 
-		OffsetLoc += FrameBoneWeights[Index] * TargetOffset;
-		TotalWeight += FrameBoneWeights[Index];
+		OffsetLoc += FrameBoneConstraints[Index].Weight * TargetOffset;
+		TotalWeight += TargetWeight;
 
 #if WITH_EDITOR
 		if (Settings.bDebugDraw)
 		{
 			// TODO: Fixup target debug to show dynamic offset goals (post-solve update?)
 			FVector PresolveBoneTarget = TargetOffset + BoneLoc;
-			PresolveTarget += FrameBoneWeights[Index] * PresolveBoneTarget;
-				
+			PresolveTarget += BoneConstraint.Weight * PresolveBoneTarget;
+			
 			DebugInfo.TargetGoals.Last().PairTargets.Add(PresolveBoneTarget);
 			DebugInfo.PairRetargetInfo.Emplace(DebugPairSpaces[Index]);
 		}
@@ -1090,7 +1168,13 @@ double FIKRetargetRelativeIKOp::ComputeTargetWeightedOffsetRow(
 #if WITH_EDITOR
 	if (Settings.bDebugDraw)
 	{
-		DebugInfo.TargetGoals.Last().Goal = PresolveTarget / TotalWeight;
+		FVector InvWeight = {
+			(FMath::IsNearlyZero(TotalWeight.X) ? 0.0 : 1.0 / TotalWeight.X),
+			(FMath::IsNearlyZero(TotalWeight.Y) ? 0.0 : 1.0 / TotalWeight.Y),
+			(FMath::IsNearlyZero(TotalWeight.Z) ? 0.0 : 1.0 / TotalWeight.Z),
+
+		};
+		DebugInfo.TargetGoals.Last().Goal = PresolveTarget * InvWeight;
 	}
 #endif // WITH_EDITOR
 	
@@ -1706,9 +1790,8 @@ void FIKRetargetRelativeIKOp::ResetCacheNotifyInfo()
 	CacheSourceDomainBones.Reset();
 	CacheBodyBones.Reset();
 	CachePairEffectBoneNames.Reset();
-	FrameBoneOffsets.Reset();
-	FrameBonePositions.Reset();
-	FrameBoneWeights.Reset();
+	FrameBoneConstraints.Reset();
+	
 	CachePinSourcePropBones.Reset();
 }
 
@@ -1727,12 +1810,10 @@ void FIKRetargetRelativeIKOp::UpdateCacheNotifyInfo(URelativeBodyBakeAnimNotify*
 	{
 		return;
 	}
-	CacheMeshPairOffset = FrameBoneOffsets.Num() / 2;
+	CacheMeshPairOffset = FrameBoneConstraints.Num() / 2;
 	
 	int32 NumBodies = CachedNotifyInfo->BodyPairs.Num();
-	FrameBoneOffsets.AddZeroed(NumBodies);
-	FrameBonePositions.AddZeroed(NumBodies);
-	FrameBoneWeights.AddZeroed(NumBodies);
+	FrameBoneConstraints.AddDefaulted(NumBodies);
 	CachePairEffectBoneNames.AddDefaulted(NumBodies);
 
 	MeshSampleRate = 30.0f;
@@ -1806,13 +1887,11 @@ void FIKRetargetRelativeIKOp::UpdateCachePropNotifyInfo(URelativePropsBakeAnimNo
 	{
 		return;
 	}
-	CachePropPairOffset = FrameBoneOffsets.Num() / 2;
+	CachePropPairOffset = FrameBoneConstraints.Num() / 2;
 
-	int32 NumBones = CachedPropNotifyInfo->PropsPairs.Num();
-	FrameBoneOffsets.AddZeroed(NumBones);
-	FrameBonePositions.AddZeroed(NumBones);
-	FrameBoneWeights.AddZeroed(NumBones);
-	CachePairEffectBoneNames.AddDefaulted(NumBones);
+	int32 NumBodies = CachedPropNotifyInfo->PropsPairs.Num();
+	FrameBoneConstraints.AddDefaulted(NumBodies);
+	CachePairEffectBoneNames.AddDefaulted(NumBodies);
 
 	PropSampleRate = 30.0f;
 	if (CachedPropNotifyInfo->PropsPairsSampleTime.Num() >= 2)
@@ -1854,6 +1933,17 @@ void FIKRetargetRelativeIKOp::UpdateCachePropNotifyInfo(URelativePropsBakeAnimNo
 		CachePairEffectBoneNames[2*OutPairIdx] = BodyBone;
 		CachePairEffectBoneNames[2*OutPairIdx + 1] = PropBone;
 
+		if (BodyBone == FName("Floor"))
+		{
+			// Prop attach bones always allowed to move
+			CacheSourceBodyEffectVertIdx[PropIdx].Add(2*OutPairIdx);
+
+			// Keep track of prop attach bones for post-solve pin/update
+			CachePinSourcePropBones.AddUnique(PropBone);
+			
+			continue;
+		}
+
 		if (PropIdx == INDEX_NONE || BodyIdx == INDEX_NONE)
 		{
 			continue;
@@ -1959,11 +2049,16 @@ void FIKRetargetRelativeIKOp::DebugDrawBodyPairs(FPrimitiveDrawInterface* InPDI,
 	const float MarkerSize = FMath::Min(MaxScale, static_cast<float>(Scale));
 	for (const FDebugDrawBodyPair& Pair : BodyPairs)
 	{
+		float Scalar = FMath::Clamp((float)Pair.Weight, 0.f, 1.f);
+		if (FMath::IsNearlyZero(Scalar))
+		{
+			continue;
+		}
 		FTransform TfmBodyA_A = FTransform(Pair.PosA_A) * BaseTransform;
 		FTransform TfmBodyB_B = FTransform(Pair.PosB_B) * BaseTransform;
 
-		IKRigDebugRendering::DrawWireCube(InPDI, TfmBodyA_A, FLinearColor::Blue,MarkerSize,MarkerSize);
-		IKRigDebugRendering::DrawWireCube(InPDI, TfmBodyB_B, FLinearColor::Red,MarkerSize,MarkerSize);
+		IKRigDebugRendering::DrawWireCube(InPDI, TfmBodyA_A, FLinearColor::Blue*Scalar,MarkerSize,MarkerSize);
+		IKRigDebugRendering::DrawWireCube(InPDI, TfmBodyB_B, FLinearColor::Red*Scalar,MarkerSize,MarkerSize);
 
 		// Possibly draw full retarget debug info
 		if (Pair.bFullPos)
diff --git a/Engine/Plugins/Experimental/Animation/RelativeIKOp/Source/RelativeIKOp/Public/RelativeIKOp.h b/Engine/Plugins/Experimental/Animation/RelativeIKOp/Source/RelativeIKOp/Public/RelativeIKOp.h
--- a/Engine/Plugins/Experimental/Animation/RelativeIKOp/Source/RelativeIKOp/Public/RelativeIKOp.h
+++ b/Engine/Plugins/Experimental/Animation/RelativeIKOp/Source/RelativeIKOp/Public/RelativeIKOp.h
@@ -59,6 +59,8 @@ struct FDebugDrawBodyPair
 	bool bFullPos;
 	FVector PosB_A;
 	FVector PosA_B;
+
+	double Weight;
 };
 
 // Debug physics body info
@@ -118,6 +120,34 @@ struct FPinBoneSettings
 	EPinBoneRotationMode RotationMode = EPinBoneRotationMode::CopyGlobalRotation;
 };
 
+UENUM()
+enum class ERelativeIKFrameBoneConstraintType : uint8
+{
+	XYZ = 0,
+	Z = 1,
+	Default = XYZ,
+};
+
+struct FRelativeIKFrameBoneConstraint
+{
+	ERelativeIKFrameBoneConstraintType ConstraintType;
+	
+	double Weight;
+	FVector Offset;
+
+	FRelativeIKFrameBoneConstraint()
+	{
+		Reset();
+	}
+
+	void Reset()
+	{
+		ConstraintType = ERelativeIKFrameBoneConstraintType::Default;
+		Weight = 0.0;
+		Offset = FVector::ZeroVector;
+	}
+};
+
 USTRUCT(BlueprintType, meta = (DisplayName = "Relative IK Settings"))
 struct FIKRetargetRelativeIKOpSettings : public FIKRetargetOpSettingsBase
 {
@@ -134,6 +164,10 @@ struct FIKRetargetRelativeIKOpSettings : public FIKRetargetOpSettingsBase
 	// Physics body source -> target name mapping 
 	UPROPERTY(EditAnywhere, BlueprintReadWrite, Category = "Source Target Mapping", meta=(ReinitializeOnEdit))
 	TMap<FName,FName> BodyMapping;
+
+	// Maximum distance for which prop-floor info is baked
+	UPROPERTY(EditAnywhere, BlueprintReadWrite, Category = "Parameters", meta=(ClampMin = "0.1"))
+	double FloorThreshold = 5.00;
 	
 	// List of bones to pin. Should include all attach bones that may be affected by relative ik bake data 
 	UPROPERTY(EditAnywhere, BlueprintReadWrite, Category = "Prop Pinning", meta=(ReinitializeOnEdit))
@@ -141,11 +175,11 @@ struct FIKRetargetRelativeIKOpSettings : public FIKRetargetOpSettingsBase
 	
 	// Maximum distance for which body pair info is baked
 	UPROPERTY(EditAnywhere, BlueprintReadWrite, Category = "Parameters", meta=(ClampMin = "0.1"))
-	double DistanceThreshold = 50.00;
+	double DistanceThreshold = 70.00;
 
 	// DistHalfPointLambda*DistanceThreshold is the place where distance weight value be 0.5
 	UPROPERTY(EditAnywhere, BlueprintReadWrite, Category = "Parameters", meta=(ClampMin = "0.01", ClampMax = "0.99"))
-	double DistHalfPointLambda = 0.5;
+	double DistHalfPointLambda = 0.15;
 
 	// DistHalfPointLambda*DistanceThreshold is the place where distance weight value be 0.5
 	UPROPERTY(EditAnywhere, BlueprintReadWrite, Category = "Parameters")
@@ -349,17 +383,17 @@ private:
 	void UpdateTargetIKGoal(
 		TArray<FIKRigGoal>& OutIKGoals,
 		FName TargetBoneName,
-		double TargetWeight,
-		const FVector& TargetGoalLoc,
+		const FVector& TargetWeight,
+		const FVector& TargetGoalOffset,
 		const FTransform& InTargetBonePose);
 
 	void UpdatePropPinTransform(
 		FName SourceBoneName,
-		double TargetWeight,
-		const FVector& TargetPropLoc,
+		const FVector& TargetWeight,
+		const FVector& TargetPropOffset,
 		const FTransform& InTargetBonePose);
 
-	double ComputeTargetWeightedOffset(
+	FVector ComputeTargetWeightedOffset(
 		FVector& OutTargetOffset,
 		const TArray<int32>& EffectIndices,
 		const FVector& BoneLoc,
@@ -372,7 +406,7 @@ private:
 		const TArray<FName>& OutBones,
 		const TArray<int32>& EffectIndices);
 
-	double ComputeTargetWeightedOffsetRow(
+	FVector ComputeTargetWeightedOffsetRow(
 		Eigen::MatrixXd& OffsetMatrix,
 		int32 RowIdx,
 		const TArray<FName>& OutBones,
@@ -477,9 +511,8 @@ private:
 	TArray<FName> CacheSourceDomainBones;
 	TSet<FName> CacheBodyBones;
 	TArray<FName> CachePairEffectBoneNames;
-	TArray<double> FrameBoneWeights;
-	TArray<FVector> FrameBonePositions;
-	TArray<FVector> FrameBoneOffsets;
+	
+	TArray<FRelativeIKFrameBoneConstraint> FrameBoneConstraints;
 
 	// Cache skel and IK goal info
 	TArray<FTransform> CacheSourceBoneInitTfm;
