RelativeIK:
+ floor constraint scaler

#Synced-CL 48616027

diff --git a/Engine/Plugins/Experimental/Animation/RelativeIKOp/Source/RelativeIKOp/Private/RelativeIKOp.cpp b/Engine/Plugins/Experimental/Animation/RelativeIKOp/Source/RelativeIKOp/Private/RelativeIKOp.cpp
--- a/Engine/Plugins/Experimental/Animation/RelativeIKOp/Source/RelativeIKOp/Private/RelativeIKOp.cpp
+++ b/Engine/Plugins/Experimental/Animation/RelativeIKOp/Source/RelativeIKOp/Private/RelativeIKOp.cpp
@@ -461,7 +461,7 @@ void FIKRetargetRelativeIKOp::ComputeFramePropSpaces(
 			double FloorOffsetZMin = Settings.DistHalfPointLambda*Settings.FloorThreshold-TargetPropCompZMin.Z;
 			FVector TargetOffsetZMin = SourcePropCompZMin - TargetPropCompZMin;
 			FVector TargetOffsetZMax = SourcePropCompZMax - TargetPropCompZMax;
-			double Weight = SourceWeight*SourceWeight + (1-SourceWeight)*TargetWeight;
+			double Weight = Settings.FloorConstraintScaler*(SourceWeight*SourceWeight + (1-SourceWeight)*TargetWeight);
 
 			FrameBoneConstraints[2*OutPairIdx].ConstraintType = ERelativeIKFrameBoneConstraintType::Z;
 			FrameBoneConstraints[2*OutPairIdx+1].ConstraintType = ERelativeIKFrameBoneConstraintType::Z;
diff --git a/Engine/Plugins/Experimental/Animation/RelativeIKOp/Source/RelativeIKOp/Public/RelativeIKOp.h b/Engine/Plugins/Experimental/Animation/RelativeIKOp/Source/RelativeIKOp/Public/RelativeIKOp.h
--- a/Engine/Plugins/Experimental/Animation/RelativeIKOp/Source/RelativeIKOp/Public/RelativeIKOp.h
+++ b/Engine/Plugins/Experimental/Animation/RelativeIKOp/Source/RelativeIKOp/Public/RelativeIKOp.h
@@ -169,6 +169,10 @@ struct FIKRetargetRelativeIKOpSettings : public FIKRetargetOpSettingsBase
 	UPROPERTY(EditAnywhere, BlueprintReadWrite, Category = "Parameters", meta=(ClampMin = "0.1"))
 	double FloorThreshold = 5.00;
 	
+	// Maximum distance for which prop-floor info is baked
+    UPROPERTY(EditAnywhere, BlueprintReadWrite, Category = "Parameters", meta=(ClampMin = "0.0"))
+    double FloorConstraintScaler = 10.00;
+	
 	// List of bones to pin. Should include all attach bones that may be affected by relative ik bake data 
 	UPROPERTY(EditAnywhere, BlueprintReadWrite, Category = "Prop Pinning", meta=(ReinitializeOnEdit))
 	TArray<FPinBoneSettings> PinPropBones;
