RelativeBodyAnimModifiers: Update prop bake bp interface to support enabling floor constraint

#Synced-CL 48642411

diff --git a/Engine/Plugins/Experimental/Animation/RelativeIKOp/Source/RelativeBodyAnimUtils/Private/RelativeBodyBlueprintFunctions.cpp b/Engine/Plugins/Experimental/Animation/RelativeIKOp/Source/RelativeBodyAnimUtils/Private/RelativeBodyBlueprintFunctions.cpp
--- a/Engine/Plugins/Experimental/Animation/RelativeIKOp/Source/RelativeBodyAnimUtils/Private/RelativeBodyBlueprintFunctions.cpp
+++ b/Engine/Plugins/Experimental/Animation/RelativeIKOp/Source/RelativeBodyAnimUtils/Private/RelativeBodyBlueprintFunctions.cpp
@@ -132,6 +132,7 @@ void URelativeBodyAnimBlueprintFunctions::BulkUpdatePropBodyModifiers(UPARAM(ref
 			UpdateModifierInst->PropsAttachBones.Add(PropSettings.AttachBoneName);
 			UpdateModifierInst->PropsOffsetTransforms.Add(PropSettings.AttachTransform);
 		}
+		UpdateModifierInst->bPropsFloorInfoBaking = Options.bPropsFloorInfoBaking;
 		UpdateModifierInst->PropsNotifyClass = Options.PropNotifyClass;
 		
 		// Always set these
diff --git a/Engine/Plugins/Experimental/Animation/RelativeIKOp/Source/RelativeBodyAnimUtils/Public/RelativeBodyBlueprintFunctions.h b/Engine/Plugins/Experimental/Animation/RelativeIKOp/Source/RelativeBodyAnimUtils/Public/RelativeBodyBlueprintFunctions.h
--- a/Engine/Plugins/Experimental/Animation/RelativeIKOp/Source/RelativeBodyAnimUtils/Public/RelativeBodyBlueprintFunctions.h
+++ b/Engine/Plugins/Experimental/Animation/RelativeIKOp/Source/RelativeBodyAnimUtils/Public/RelativeBodyBlueprintFunctions.h
@@ -58,10 +58,14 @@ struct RELATIVEBODYANIMUTILS_API FPropBodyAnimModifierOptions : public FRelative
 {
 	GENERATED_BODY()
 
-	// Settings for eash prop static mesh attachment to bake
+	// Settings for each prop static mesh attachment to bake
 	UPROPERTY(EditAnywhere, BlueprintReadWrite, Category = "Settings")
 	TArray<FPropMeshSettings> PropMeshSettings;
 	
+	// Include floor constraint info in prop bake
+	UPROPERTY(EditAnywhere, BlueprintReadWrite, Category = "Settings")
+	bool bPropsFloorInfoBaking = false;
+	
 	// Prop Notify subclass to create
 	UPROPERTY(EditAnywhere, BlueprintReadWrite, Category = "Settings")
 	TSubclassOf<URelativePropsBakeAnimNotify> PropNotifyClass = nullptr;
