RelativeIKOp: Fix for possible crash in RunAfterParent and make bDryRun proprly passthrough pin bones.

#Synced-CL 48670229

diff --git a/Engine/Plugins/Experimental/Animation/RelativeIKOp/Source/RelativeIKOp/Private/RelativeIKOp.cpp b/Engine/Plugins/Experimental/Animation/RelativeIKOp/Source/RelativeIKOp/Private/RelativeIKOp.cpp
--- a/Engine/Plugins/Experimental/Animation/RelativeIKOp/Source/RelativeIKOp/Private/RelativeIKOp.cpp
+++ b/Engine/Plugins/Experimental/Animation/RelativeIKOp/Source/RelativeIKOp/Private/RelativeIKOp.cpp
@@ -224,7 +224,7 @@ void FIKRetargetRelativeIKOp::RunAfterParent(
 		
 		int32 TargetBoneIdx = CacheTargetSkelIndices[TargetPinBone];
 		int32 CachePropIdx = CachePinSourcePropBones.Find(SourcePinBone);
-		if (CachePropIdx != INDEX_NONE)
+		if (CachePinTargetPropTransforms.IsValidIndex(CachePropIdx))
 		{
 			OutTargetGlobalPose[TargetBoneIdx] = CachePinTargetPropTransforms[CachePropIdx];
 		}
@@ -1006,6 +1006,12 @@ void FIKRetargetRelativeIKOp::UpdatePropPinTransform(
 	const FVector& TargetPropOffset,
 	const FTransform& InTargetBonePose)
 {
+	// Disable prop bone changes if debug drawing and dry run are on
+	if (!Settings.bDebugDraw || !Settings.bDryRun)
+	{
+		return;
+	}
+	
 	int32 PinPropIdx = CachePinSourcePropBones.Find(SourceBoneName);
 	if (PinPropIdx == INDEX_NONE)
 	{
