RelativeIKOp: Fixed dry run code from previous CL

#Synced-CL 48670229

diff --git a/Engine/Plugins/Experimental/Animation/RelativeIKOp/Source/RelativeIKOp/Private/RelativeIKOp.cpp b/Engine/Plugins/Experimental/Animation/RelativeIKOp/Source/RelativeIKOp/Private/RelativeIKOp.cpp
--- a/Engine/Plugins/Experimental/Animation/RelativeIKOp/Source/RelativeIKOp/Private/RelativeIKOp.cpp
+++ b/Engine/Plugins/Experimental/Animation/RelativeIKOp/Source/RelativeIKOp/Private/RelativeIKOp.cpp
@@ -1007,7 +1007,7 @@ void FIKRetargetRelativeIKOp::UpdatePropPinTransform(
 	const FTransform& InTargetBonePose)
 {
 	// Disable prop bone changes if debug drawing and dry run are on
-	if (!Settings.bDebugDraw || !Settings.bDryRun)
+	if (Settings.bDebugDraw && Settings.bDryRun)
 	{
 		return;
 	}
