RelativeIK:
+ PropModifier support for skeletal  mesh
+ PropNotifier domainbody bake

#Synced-CL 48826176

diff --git a/Engine/Plugins/Experimental/Animation/RelativeIKOp/Source/RelativeBodyAnimInfo/Public/RelativeBodyAnimNotifies.h b/Engine/Plugins/Experimental/Animation/RelativeIKOp/Source/RelativeBodyAnimInfo/Public/RelativeBodyAnimNotifies.h
--- a/Engine/Plugins/Experimental/Animation/RelativeIKOp/Source/RelativeBodyAnimInfo/Public/RelativeBodyAnimNotifies.h
+++ b/Engine/Plugins/Experimental/Animation/RelativeIKOp/Source/RelativeBodyAnimInfo/Public/RelativeBodyAnimNotifies.h
@@ -183,6 +183,9 @@ public:
 	
 	UPROPERTY(VisibleAnywhere, AdvancedDisplay, meta = (Category = "Detail"))
 	TArray<FVector3f> PropsPairsLocalReference;
+
+	UPROPERTY(VisibleAnywhere, AdvancedDisplay, meta = (Category = "Detail"))
+	TArray<bool> bPropsPairsIsParentDominates;
 };
 
 #undef UE_API
diff --git a/Engine/Plugins/Experimental/Animation/RelativeIKOp/Source/RelativeBodyAnimUtils/Private/RelativeBodyBlueprintFunctions.cpp b/Engine/Plugins/Experimental/Animation/RelativeIKOp/Source/RelativeBodyAnimUtils/Private/RelativeBodyBlueprintFunctions.cpp
--- a/Engine/Plugins/Experimental/Animation/RelativeIKOp/Source/RelativeBodyAnimUtils/Private/RelativeBodyBlueprintFunctions.cpp
+++ b/Engine/Plugins/Experimental/Animation/RelativeIKOp/Source/RelativeBodyAnimUtils/Private/RelativeBodyBlueprintFunctions.cpp
@@ -123,14 +123,10 @@ void URelativeBodyAnimBlueprintFunctions::BulkUpdatePropBodyModifiers(UPARAM(ref
 			continue;
 		}
 		
-		UpdateModifierInst->PropsStaticMeshs.Empty(Options.PropMeshSettings.Num());
-		UpdateModifierInst->PropsAttachBones.Empty(Options.PropMeshSettings.Num());
-		UpdateModifierInst->PropsOffsetTransforms.Empty(Options.PropMeshSettings.Num());
-		for (const FPropMeshSettings& PropSettings : Options.PropMeshSettings)
+		UpdateModifierInst->PropsData.Empty(Options.PropMeshSettings.Num());
+		for (const FPropsInfo& PropInfo : Options.PropMeshSettings)
 		{
-			UpdateModifierInst->PropsStaticMeshs.Add(PropSettings.StaticMesh);
-			UpdateModifierInst->PropsAttachBones.Add(PropSettings.AttachBoneName);
-			UpdateModifierInst->PropsOffsetTransforms.Add(PropSettings.AttachTransform);
+			UpdateModifierInst->PropsData.Add(PropInfo);
 		}
 		UpdateModifierInst->bPropsFloorInfoBaking = Options.bPropsFloorInfoBaking;
 		UpdateModifierInst->PropsNotifyClass = Options.PropNotifyClass;
diff --git a/Engine/Plugins/Experimental/Animation/RelativeIKOp/Source/RelativeBodyAnimUtils/Private/RelativePropsAnimModifier.cpp b/Engine/Plugins/Experimental/Animation/RelativeIKOp/Source/RelativeBodyAnimUtils/Private/RelativePropsAnimModifier.cpp
--- a/Engine/Plugins/Experimental/Animation/RelativeIKOp/Source/RelativeBodyAnimUtils/Private/RelativePropsAnimModifier.cpp
+++ b/Engine/Plugins/Experimental/Animation/RelativeIKOp/Source/RelativeBodyAnimUtils/Private/RelativePropsAnimModifier.cpp
@@ -78,6 +78,41 @@ struct FPropsMeshAndTree
 	}
 };
 
+const FPositionVertexBuffer& FPropsInfo::GetPositionVertexBuffer()
+{
+	if (PropSkeletalMeshAsset != nullptr)
+	{
+		const FSkeletalMeshRenderData* RenderData = PropSkeletalMeshAsset->GetResourceForRendering();
+		const FSkeletalMeshLODRenderData& LODData = RenderData->LODRenderData[0];
+		return LODData.StaticVertexBuffers.PositionVertexBuffer;
+	}
+
+	if (PropStaticMeshAsset != nullptr)
+	{
+		const FStaticMeshLODResources& LODResources = PropStaticMeshAsset->GetRenderData()->LODResources[0];
+		return LODResources.VertexBuffers.PositionVertexBuffer;
+	}
+
+	// Fallback: return an empty (static) vertex buffer to avoid crashes
+	static FPositionVertexBuffer EmptyBuffer;
+	return EmptyBuffer;
+}
+
+FMeshDescription* FPropsInfo::GetMeshDescription(int32 LodIndex)
+{
+	if (PropSkeletalMeshAsset != nullptr)
+	{
+		return PropSkeletalMeshAsset->GetMeshDescription(LodIndex);
+	}
+	
+	if (PropStaticMeshAsset != nullptr)
+	{
+		return PropStaticMeshAsset->GetMeshDescription(LodIndex);
+	}
+
+	return nullptr;
+}
+
 void URelativePropsAnimModifier::OnApply_Implementation(UAnimSequence* InAnimation)
 {
 	if (NotifyClass)
@@ -158,7 +193,7 @@ void URelativePropsAnimModifier::OnApply_Implementation(UAnimSequence* InAnimati
 
 		bool bUseDenseNotifyInfo = PropsNotifyClass->IsChildOf(URelativePropsBakeAnimNotify::StaticClass());
 		TArray<FPropsMeshAndTree> PropsTreeData;
-		PropsTreeData.SetNum(PropsStaticMeshs.Num());
+		PropsTreeData.SetNum(PropsData.Num());
 		if (bUseDenseNotifyInfo)
 		{
 			FName TrackName("RelProps-<" + SkeletalMeshAsset->GetName() + ">");
@@ -195,10 +230,10 @@ void URelativePropsAnimModifier::OnApply_Implementation(UAnimSequence* InAnimati
 				BakedAnimNotifyInfo->OffsetTransformsForBones.Emplace(SourceBodyBoneName, BodyTransform1);
 			}
 				
-			for (size_t PropsIdx = 0; PropsIdx < PropsStaticMeshs.Num(); PropsIdx++)
+			for (size_t PropsIdx = 0; PropsIdx < PropsData.Num(); PropsIdx++)
 			{
-				BakedAnimNotifyInfo->OffsetTransformsForBones.Emplace(PropsAttachBones[PropsIdx], PropsOffsetTransforms[PropsIdx]);
-				FMeshDescription* PropsMeshDescription = PropsStaticMeshs[PropsIdx]->GetMeshDescription(0); // LOD 0
+				BakedAnimNotifyInfo->OffsetTransformsForBones.Emplace(PropsData[PropsIdx].AttachBone.BoneName,  PropsData[PropsIdx].AttachTransform);
+				FMeshDescription* PropsMeshDescription = PropsData[PropsIdx].GetMeshDescription(0); // LOD 0
 				if (PropsMeshDescription)
 				{
 					PropsTreeData[PropsIdx].Init(PropsMeshDescription);
@@ -211,13 +246,15 @@ void URelativePropsAnimModifier::OnApply_Implementation(UAnimSequence* InAnimati
 						continue;
 					}
 					FName BoneName = RefSkeleton.GetBoneName(BodyIndicesToBoneIndices[BodyIndex1]);
-					BakedAnimNotifyInfo->PropsPairs.Add(PropsAttachBones[PropsIdx]);
+					BakedAnimNotifyInfo->PropsPairs.Add(PropsData[PropsIdx].AttachBone.BoneName);
 					BakedAnimNotifyInfo->PropsPairs.Add(BoneName);
+					BakedAnimNotifyInfo->bPropsPairsIsParentDominates.Add(IsDomainBodyIndices[BodyIndex1]);
 				}
 				if (bPropsFloorInfoBaking)
 				{
-					BakedAnimNotifyInfo->PropsPairs.Add(PropsAttachBones[PropsIdx]);
+					BakedAnimNotifyInfo->PropsPairs.Add(PropsData[PropsIdx].AttachBone.BoneName);
 					BakedAnimNotifyInfo->PropsPairs.Add(FName("Floor"));
+					BakedAnimNotifyInfo->bPropsPairsIsParentDominates.Add(true);
 				}
 			}
 			BakedAnimNotifyInfo->PropsPairsLocalReference.SetNumZeroed(BakedAnimNotifyInfo->PropsPairs.Num() * SampleNum);
@@ -263,12 +300,11 @@ void URelativePropsAnimModifier::OnApply_Implementation(UAnimSequence* InAnimati
 			else 
 			{
 				BakedAnimNotifyInfo->PropsPairsSampleTime[SampleIndex] = SampleTime;
-				for (int32 PropsIdx = 0; PropsIdx < PropsStaticMeshs.Num(); PropsIdx++)
+				for (int32 PropsIdx = 0; PropsIdx < PropsData.Num(); PropsIdx++)
 				{
-					FTransform AttachBoneTransform = UAnimPoseExtensions::GetBonePose(AnimPose, PropsAttachBones[PropsIdx], EAnimPoseSpaces::World);
-					FTransform PropsLocalToWorld = PropsOffsetTransforms[PropsIdx]*AttachBoneTransform;
-					const FStaticMeshLODResources& LODResources = PropsStaticMeshs[PropsIdx]->GetRenderData()->LODResources[0];
-					const FPositionVertexBuffer& PositionVertexBuffer = LODResources.VertexBuffers.PositionVertexBuffer;
+					FTransform AttachBoneTransform = UAnimPoseExtensions::GetBonePose(AnimPose, PropsData[PropsIdx].AttachBone.BoneName, EAnimPoseSpaces::World);
+					FTransform PropsLocalToWorld = PropsData[PropsIdx].AttachTransform*AttachBoneTransform;
+					const FPositionVertexBuffer& PositionVertexBuffer = PropsData[PropsIdx].GetPositionVertexBuffer();
 					for (int32 BodyIndex = 0; BodyIndex < NumBodies; ++BodyIndex)
 					{
 						if (BodyIndicesToIgnore[BodyIndex])
@@ -299,7 +335,7 @@ void URelativePropsAnimModifier::OnApply_Implementation(UAnimSequence* InAnimati
 						// UE_LOG(LogAnimation, Warning, TEXT("MinDIst %f: (%f,%f,%f)."), MinDist, NearestPoint.X, NearestPoint.Y, NearestPoint.Z);
 						if (NearestQueryVertexIndex != INDEX_NONE)
 						{
-							FName BoneName1 = PropsAttachBones[PropsIdx];
+							FName BoneName1 = PropsData[PropsIdx].AttachBone.BoneName;
 							FName BoneName2 = RefSkeleton.GetBoneName(BodyIndicesToBoneIndices[BodyIndex]);
 							FVector3f ContactPoint1(PropsLocalToWorld.TransformPosition(NearestPoint));
 							FVector3f ContactPoint2 = VLocations[VertexIndicesInfluencedByBodyIndices[BodyIndex][NearestQueryVertexIndex]];
@@ -319,7 +355,7 @@ void URelativePropsAnimModifier::OnApply_Implementation(UAnimSequence* InAnimati
 					}
 					if (bPropsFloorInfoBaking)
 					{
-						FTransform3f PropOffsetTfm(BakedAnimNotifyInfo->OffsetTransformsForBones[PropsAttachBones[PropsIdx]]);
+						FTransform3f PropOffsetTfm(BakedAnimNotifyInfo->OffsetTransformsForBones[PropsData[PropsIdx].AttachBone.BoneName]);
 						int32 Offset = SampleIndex * BakedAnimNotifyInfo->PropsPairs.Num() + Count * 2;
 						int32 LowestVertexIndex = FindLowestIndex(PropsVLocations[PropsIdx]);
 						int32 HighestVertexIndex=FindHighestIndex(PropsVLocations[PropsIdx]);
@@ -338,18 +374,17 @@ void URelativePropsAnimModifier::OnApply_Implementation(UAnimSequence* InAnimati
 void URelativePropsAnimModifier::GetPropsVertices(TArray<TArray<FVector3f>>& PropsVLocations, const FAnimPose& AnimPose)
 {
 	PropsVLocations.Empty();
-	for (int32 PropsIdx = 0; PropsIdx < PropsStaticMeshs.Num(); PropsIdx++)
+	for (int32 PropsIdx = 0; PropsIdx < PropsData.Num(); PropsIdx++)
 	{
-		const FTransform AttachBoneTransform = UAnimPoseExtensions::GetBonePose(AnimPose, PropsAttachBones[PropsIdx], EAnimPoseSpaces::World);
-
-		const FStaticMeshLODResources& LODResources = PropsStaticMeshs[PropsIdx]->GetRenderData()->LODResources[0];
-		const FPositionVertexBuffer& PositionVertexBuffer = LODResources.VertexBuffers.PositionVertexBuffer;
+		const FTransform AttachBoneTransform = UAnimPoseExtensions::GetBonePose(AnimPose, PropsData[PropsIdx].AttachBone.BoneName, EAnimPoseSpaces::World);
+		
+		const FPositionVertexBuffer& PositionVertexBuffer = PropsData[PropsIdx].GetPositionVertexBuffer();
 
 		TArray<FVector3f> LocalVertices;
 		int32 NumVerts = PositionVertexBuffer.GetNumVertices();
 		for (int32 i = 0; i < NumVerts; ++i)
 		{
-			FVector LocalPos = PropsOffsetTransforms[PropsIdx].TransformPosition(FVector(PositionVertexBuffer.VertexPosition(i)));
+			FVector LocalPos = PropsData[PropsIdx].AttachTransform.TransformPosition(FVector(PositionVertexBuffer.VertexPosition(i)));
 			LocalVertices.Add(FVector3f(AttachBoneTransform.TransformPosition(LocalPos)));
 		}
 		PropsVLocations.Add(LocalVertices);
diff --git a/Engine/Plugins/Experimental/Animation/RelativeIKOp/Source/RelativeBodyAnimUtils/Public/RelativeBodyBlueprintFunctions.h b/Engine/Plugins/Experimental/Animation/RelativeIKOp/Source/RelativeBodyAnimUtils/Public/RelativeBodyBlueprintFunctions.h
--- a/Engine/Plugins/Experimental/Animation/RelativeIKOp/Source/RelativeBodyAnimUtils/Public/RelativeBodyBlueprintFunctions.h
+++ b/Engine/Plugins/Experimental/Animation/RelativeIKOp/Source/RelativeBodyAnimUtils/Public/RelativeBodyBlueprintFunctions.h
@@ -7,6 +7,8 @@
 
 #include "RelativeBodyBlueprintFunctions.generated.h"
 
+struct FPropsInfo;
+
 USTRUCT(BlueprintType)
 struct RELATIVEBODYANIMUTILS_API FRelativeBodyAnimModifierOptions
 {
@@ -38,21 +40,6 @@ struct RELATIVEBODYANIMUTILS_API FRelativeBodyAnimModifierOptions
 	TSubclassOf<URelativeBodyAnimNotifyBase> NotifyClass = nullptr;
 };
 
-USTRUCT(BlueprintType)
-struct FPropMeshSettings
-{
-	GENERATED_BODY()
-	
-	UPROPERTY(EditAnywhere, BlueprintReadWrite, Category = "Settings")
-	TObjectPtr<UStaticMesh> StaticMesh;
-
-	UPROPERTY(EditAnywhere, BlueprintReadWrite, Category = "Settings")
-	FName AttachBoneName;
-	
-	UPROPERTY(EditAnywhere, BlueprintReadWrite, Category = "Settings")
-	FTransform AttachTransform;
-};
-
 USTRUCT(BlueprintType)
 struct RELATIVEBODYANIMUTILS_API FPropBodyAnimModifierOptions : public FRelativeBodyAnimModifierOptions
 {
@@ -60,7 +47,7 @@ struct RELATIVEBODYANIMUTILS_API FPropBodyAnimModifierOptions : public FRelative
 
 	// Settings for each prop static mesh attachment to bake
 	UPROPERTY(EditAnywhere, BlueprintReadWrite, Category = "Settings")
-	TArray<FPropMeshSettings> PropMeshSettings;
+	TArray<FPropsInfo> PropMeshSettings;
 	
 	// Include floor constraint info in prop bake
 	UPROPERTY(EditAnywhere, BlueprintReadWrite, Category = "Settings")
diff --git a/Engine/Plugins/Experimental/Animation/RelativeIKOp/Source/RelativeBodyAnimUtils/Public/RelativePropsAnimModifier.h b/Engine/Plugins/Experimental/Animation/RelativeIKOp/Source/RelativeBodyAnimUtils/Public/RelativePropsAnimModifier.h
--- a/Engine/Plugins/Experimental/Animation/RelativeIKOp/Source/RelativeBodyAnimUtils/Public/RelativePropsAnimModifier.h
+++ b/Engine/Plugins/Experimental/Animation/RelativeIKOp/Source/RelativeBodyAnimUtils/Public/RelativePropsAnimModifier.h
@@ -14,6 +14,34 @@ class USkeletalMesh;
 struct FAnimPose;
 struct FReferenceSkeleton;
 
+class FPositionVertexBuffer;
+struct FMeshDescription;
+
+USTRUCT(BlueprintType)
+struct FPropsInfo
+{
+	GENERATED_BODY()
+	// Prop mesh to preview
+	UPROPERTY(EditAnywhere, BlueprintReadWrite, Category = "Prop", meta=(ReinitializeOnEdit))
+	TObjectPtr<UStaticMesh> PropStaticMeshAsset;
+	
+	// Prop mesh to preview
+	UPROPERTY(EditAnywhere, BlueprintReadWrite, Category = "Prop", meta=(ReinitializeOnEdit))
+	TObjectPtr<USkeletalMesh> PropSkeletalMeshAsset;
+
+	// Source skeletal mesh bone to attach prop to
+	UPROPERTY(EditAnywhere, Category = "Prop", meta=(ReinitializeOnEdit))
+	FBoneReference AttachBone;
+
+	// Local transform of prop relative to attach bone
+	UPROPERTY(EditAnywhere, BlueprintReadWrite, Category = "Prop", meta=(ReinitializeOnEdit))
+	FTransform AttachTransform;
+
+public:
+	const FPositionVertexBuffer& GetPositionVertexBuffer();
+	FMeshDescription* GetMeshDescription(int32 LodIndex);
+};
+
 /**
  * Animation modifier for baking relative body relationships into anim notify
  */
@@ -28,16 +56,8 @@ public:
 	/** End UAnimationModifier interface */
 
 public:
-	
-	UPROPERTY(EditAnywhere, BlueprintReadOnly, Category = "Settings", meta = (TitleProperty = "PropsStaticMeshs"))
-	TArray<TObjectPtr<UStaticMesh>> PropsStaticMeshs;
-
-	UPROPERTY(EditAnywhere, BlueprintReadOnly, Category = "Settings", meta = (TitleProperty = "PropsAttachBones"))
-	TArray<FName> PropsAttachBones;
-
-	/** The local offset transform that defines the attachment position on the props static mesh */
-	UPROPERTY(EditAnywhere, BlueprintReadOnly, Category = "Settings", meta = (TitleProperty = "PropsOffsetTransforms"))
-	TArray<FTransform> PropsOffsetTransforms;
+	UPROPERTY(EditAnywhere, BlueprintReadOnly, Category = "Settings", meta = (TitleProperty = "PropsData"))
+	TArray<FPropsInfo> PropsData;
 
 	/**  Relative Props Notify subclass to create */
 	UPROPERTY(EditAnywhere, Category = "Settings", meta = (TitleProperty = "PropsNotifyClass"))
diff --git a/Engine/Plugins/Experimental/Animation/RelativeIKOp/Source/RelativeIKOp/Private/RelativeIKOp.cpp b/Engine/Plugins/Experimental/Animation/RelativeIKOp/Source/RelativeIKOp/Private/RelativeIKOp.cpp
--- a/Engine/Plugins/Experimental/Animation/RelativeIKOp/Source/RelativeIKOp/Private/RelativeIKOp.cpp
+++ b/Engine/Plugins/Experimental/Animation/RelativeIKOp/Source/RelativeIKOp/Private/RelativeIKOp.cpp
@@ -133,7 +133,7 @@ void FIKRetargetRelativeIKOp::Run(
 	TArray<FDebugRelativeTargetPairSpace> LocalTargetPairSpaces;
 
 	int32 NumMeshPairs = (CachedNotifyInfo) ? CachedNotifyInfo->bBodyPairsIsParentDominates.Num() : 0;
-	int32 NumPropPairs = (CachedPropNotifyInfo) ? (CachedPropNotifyInfo->PropsPairs.Num() / 2) : 0;
+	int32 NumPropPairs = (CachedPropNotifyInfo) ? CachedPropNotifyInfo->bPropsPairsIsParentDominates.Num() : 0;
 	int32 TotalPairs = NumMeshPairs + NumPropPairs;
 	
 	int32 NumMeshBodies = (CachedNotifyInfo) ? CachedNotifyInfo->BodyPairs.Num() : 0;
@@ -179,9 +179,7 @@ void FIKRetargetRelativeIKOp::Run(
 	TArray<FIKRigGoal>& IKGoals = GoalContainer.GetGoalArray();
 	
 	LocalDebugInfo.TargetGoals.Reserve(CacheSourceBodyEffectVertIdx.Num());
-
-	// TODO: This is a hack to handle the fact that domain bodies aren't differentiated in prop bake data
-	DropPropDomainBodyEffects();
+	
 	// Update bone goals to relative IK positions
 	if (Settings.bMultiBoneSolve)
 	{
@@ -193,10 +191,28 @@ void FIKRetargetRelativeIKOp::Run(
 	}
 
 #if WITH_EDITOR
-	LocalDebugInfo.TargetDomainBones.SetNum(CacheSourceDomainBones.Num());
-	for (int i=0; i<CacheSourceDomainBones.Num(); i++)
+	LocalDebugInfo.TargetDomainBones.Reset(CacheSourceEffectBones.Num());
+	for (int PairIdx=0; PairIdx<NumMeshPairs; PairIdx++)
 	{
-		LocalDebugInfo.TargetDomainBones[i] = ApplyBodyMap(CacheSourceDomainBones[i]);
+		if(CachedNotifyInfo->bBodyPairsIsParentDominates[PairIdx])
+		{
+			FName DomainBody = ApplyBodyMap(CachedNotifyInfo->BodyPairs[2*PairIdx]);
+			if (!LocalDebugInfo.TargetDomainBones.Contains(DomainBody))
+			{
+				LocalDebugInfo.TargetDomainBones.Add(DomainBody);
+			}
+		}
+	}
+	for (int PairIdx=0; PairIdx<NumPropPairs; PairIdx++)
+	{
+		if(CachedPropNotifyInfo->bPropsPairsIsParentDominates[PairIdx])
+		{
+			FName DomainBody = ApplyBodyMap(CachedPropNotifyInfo->PropsPairs[2*PairIdx+1]);
+			if (!LocalDebugInfo.TargetDomainBones.Contains(DomainBody))
+			{
+				LocalDebugInfo.TargetDomainBones.Add(DomainBody);
+			}
+		}
 	}
 	UpdateDebugInfo(SourceScale, InSourceGlobalPose, OutTargetGlobalPose, LocalDebugInfo);
 #endif
@@ -395,7 +411,7 @@ void FIKRetargetRelativeIKOp::ComputeFramePropSpaces(
 	FDebugRelativeIKDrawInfo& DebugInfo,
 	TArray<FDebugRelativeTargetPairSpace>& DebugPairSpaces)
 {
-	int32 PairCount = CachedPropNotifyInfo->PropsPairs.Num() / 2;
+	int32 PairCount = CachedPropNotifyInfo->bPropsPairsIsParentDominates.Num();
 	if (PairCount <= 0)
 	{
 		return;
@@ -496,6 +512,8 @@ void FIKRetargetRelativeIKOp::ComputeFramePropSpaces(
 			continue;
 		}
 		
+		
+		
 		int32 SourcePropPoseIdx = CacheSourceSkelIndices[SourcePropBone];
 		int32 SourceBodyPoseIdx = CacheSourceSkelIndices[SourceBodyBone];
 		
@@ -505,7 +523,7 @@ void FIKRetargetRelativeIKOp::ComputeFramePropSpaces(
 		FVector SourceBoneBakePropV(FramePointView[2*PairIdx]);
 		FVector SourceBoneBakeBodyV(FramePointView[2*PairIdx+1]);
 
-		bool bBodyMovable = CacheSourceDomainBones.Contains(SourceBodyBone);
+		bool bBodyMovable = CachedPropNotifyInfo->bPropsPairsIsParentDominates[PairIdx];
 		
 		FTransform SourcePropTfm = SourceScaleTfm * SourcePose[SourcePropPoseIdx];
 		FBodyTransform SourceBodyTfm;
@@ -598,19 +616,6 @@ void FIKRetargetRelativeIKOp::ComputeFramePropSpaces(
 	}
 }
 
-void FIKRetargetRelativeIKOp::DropPropDomainBodyEffects()
-{
-	// TODO: This should be handled by baking necessary info in prop/body notifies
-	for (int32 Index=0; Index < CacheSourceEffectBones.Num(); ++Index)
-	{
-		const FName SourceBoneName = CacheSourceEffectBones[Index]; 
-		if (CacheSourceDomainBones.Contains(SourceBoneName))
-		{
-			CacheSourceBodyEffectVertIdx[Index].Reset();
-		}
-	}
-}
-
 void FIKRetargetRelativeIKOp::UpdatePinData(const FIKRetargetProcessor& Processor)
 {
 	// get skeletons we are copying from/to
@@ -1793,7 +1798,7 @@ void FIKRetargetRelativeIKOp::ResetCacheNotifyInfo()
 	CacheMapBoneSourceTargetTfm.Reset();
 	CacheSourceBodyEffectVertIdx.Reset();
 	CacheSourceEffectBones.Reset();
-	CacheSourceDomainBones.Reset();
+	
 	CacheBodyBones.Reset();
 	CachePairEffectBoneNames.Reset();
 	FrameBoneConstraints.Reset();
@@ -1866,12 +1871,7 @@ void FIKRetargetRelativeIKOp::UpdateCacheNotifyInfo(URelativeBodyBakeAnimNotify*
 			// Primary bone can be affect if e.g. hands clapping
 			CacheSourceBodyEffectVertIdx[PrimaryIdx].Add(2*OutPairIdx);
 		}
-		else
-		{
-			// TODO: Bake necessary info in props to avoid having to store this data
-			CacheSourceDomainBones.Add(PrimaryBone);
-		}
-
+		
 		// TODO: Normalize pair edge info/frame info
 		CachePairEffectBoneNames[2*OutPairIdx] = OtherBone;
 		CachePairEffectBoneNames[2*OutPairIdx + 1] = PrimaryBone;
@@ -1880,7 +1880,7 @@ void FIKRetargetRelativeIKOp::UpdateCacheNotifyInfo(URelativeBodyBakeAnimNotify*
 
 void FIKRetargetRelativeIKOp::UpdateCachePropNotifyInfo(URelativePropsBakeAnimNotify* NotifyInfo)
 {
-	if (NotifyInfo->OffsetTransformsForBones.IsEmpty())
+	if (NotifyInfo->bPropsPairsIsParentDominates.IsEmpty())
 	{
 		return;
 	}
@@ -1924,10 +1924,10 @@ void FIKRetargetRelativeIKOp::UpdateCachePropNotifyInfo(URelativePropsBakeAnimNo
 	}
 
 	// Cache bone-bone effectors for creating weighted goals
-	const int32 PairCount = CachedPropNotifyInfo->PropsPairs.Num() / 2;
-	for (int PairIdx = 0; PairIdx < PairCount; ++PairIdx)
+	for (int PairIdx = 0; PairIdx < CachedPropNotifyInfo->bPropsPairsIsParentDominates.Num(); ++PairIdx)
 	{
 		int32 OutPairIdx = PairIdx + CachePropPairOffset;
+		bool bBodyFixed = CachedPropNotifyInfo->bPropsPairsIsParentDominates[PairIdx];
 		
 		FName PropBone = CachedPropNotifyInfo->PropsPairs[2*PairIdx];
 		FName BodyBone = CachedPropNotifyInfo->PropsPairs[2*PairIdx + 1];
@@ -1954,11 +1954,14 @@ void FIKRetargetRelativeIKOp::UpdateCachePropNotifyInfo(URelativePropsBakeAnimNo
 		{
 			continue;
 		}
-
-		// TODO: Only add IK goal bones here!
-		CacheSourceBodyEffectVertIdx[BodyIdx].Add(2*OutPairIdx+1);
+		
 		// Prop attach bones always allowed to move
 		CacheSourceBodyEffectVertIdx[PropIdx].Add(2*OutPairIdx);
+		if (!bBodyFixed)
+		{
+			// TODO: Only add IK goal bones here!
+			CacheSourceBodyEffectVertIdx[BodyIdx].Add(2*OutPairIdx+1);
+		}
 
 		// Keep track of prop attach bones for post-solve pin/update
 		CachePinSourcePropBones.AddUnique(PropBone);
diff --git a/Engine/Plugins/Experimental/Animation/RelativeIKOp/Source/RelativeIKOp/Public/RelativeIKOp.h b/Engine/Plugins/Experimental/Animation/RelativeIKOp/Source/RelativeIKOp/Public/RelativeIKOp.h
--- a/Engine/Plugins/Experimental/Animation/RelativeIKOp/Source/RelativeIKOp/Public/RelativeIKOp.h
+++ b/Engine/Plugins/Experimental/Animation/RelativeIKOp/Source/RelativeIKOp/Public/RelativeIKOp.h
@@ -323,8 +323,6 @@ private:
 	void UpdateCacheNotifyInfo(URelativeBodyBakeAnimNotify* NotifyInfo);
 	void UpdateCacheChainInfo(const TArray<FTransform>& TargetPose, FName BoneName);
 	void UpdateCachePropNotifyInfo(URelativePropsBakeAnimNotify* NotifyInfo);
-
-	void DropPropDomainBodyEffects();
 	
 	struct FRelativeIKPinBoneData
 	{
@@ -512,7 +510,6 @@ private:
 	// Array of body-vert indices affecting each joint goal (for weighted averaging)
 	TArray<TArray<int32>> CacheSourceBodyEffectVertIdx;
 	TArray<FName> CacheSourceEffectBones;
-	TArray<FName> CacheSourceDomainBones;
 	TSet<FName> CacheBodyBones;
 	TArray<FName> CachePairEffectBoneNames;
 	
