PreviewPropOp: Add playhead computation to support playing prop anim sequences in preview.

#Synced-CL 49023493

diff --git a/Engine/Plugins/Experimental/Animation/RelativeIKOp/Source/PreviewPropOp/Private/PreviewPropOp.cpp b/Engine/Plugins/Experimental/Animation/RelativeIKOp/Source/PreviewPropOp/Private/PreviewPropOp.cpp
--- a/Engine/Plugins/Experimental/Animation/RelativeIKOp/Source/PreviewPropOp/Private/PreviewPropOp.cpp
+++ b/Engine/Plugins/Experimental/Animation/RelativeIKOp/Source/PreviewPropOp/Private/PreviewPropOp.cpp
@@ -7,11 +7,17 @@
 
 #if WITH_EDITOR
 #include "PrimitiveDrawingUtils.h"
+#include "Animation/AnimInstance.h"
+#include "Animation/AnimMontage.h"
+#include "Animation/AnimSequence.h"
+#include "Animation/AnimSequenceBase.h"
+#include "Components/SkeletalMeshComponent.h"
 #include "Engine/EngineTypes.h"
 #endif
 
 #include "StaticMeshAttributes.h"
 #include "AI/NavigationModifier.h"
+#include "GeometryCollection/Facades/CollectionUVFacade.h"
 
 #include UE_INLINE_GENERATED_CPP_BY_NAME(PreviewPropOp)
 
@@ -143,6 +149,10 @@ void FIKRetargetPreviewPropOp::Run(
 			FTransform TargetBoneTfm = OutTargetGlobalPose[PropMesh.TargetAttachBoneIndex];
 			PropInfo.TargetTfm = PropMesh.AttachTransform * TargetBoneTfm;
 		}
+		
+		// TODO: Move this into processing for props w/ anim seq only 
+		float TestPropAnimLength = 0.5;
+		float TestPropPlayhead = ComputePropAnimPlayhead(TestPropAnimLength, true);
 
 		++PropIdx;
 	}
@@ -150,6 +160,34 @@ void FIKRetargetPreviewPropOp::Run(
 #endif // WITH_EDITOR
 }
 
+void FIKRetargetPreviewPropOp::AnimGraphPreUpdateMainThread(USkeletalMeshComponent& SourceMeshComponent, USkeletalMeshComponent& TargetMeshComponent)
+{
+#if WITH_EDITOR
+	if (!bIsInitialized)
+	{
+		return;
+	}
+
+	UAnimInstance* SourceAnimInstance = SourceMeshComponent.GetAnimInstance();
+	if (!SourceAnimInstance)
+	{
+		return;
+	}
+	
+	// Use full montage time
+	const FAnimMontageInstance* MontageInstance = SourceAnimInstance->GetActiveMontageInstance();
+	if (MontageInstance)
+	{
+		// TODO: DeltaTimeRecord doesn't seem to work for montages when scrubbing
+		CacheSourceCharacterAnim = MontageInstance->Montage;
+		SourceAnimPlayhead = MontageInstance->GetPosition();
+		return;
+	}
+	
+	UpdateAnimSeqPlayhead(SourceAnimInstance);
+#endif //WITH_EDITOR
+}
+
 FIKRetargetOpSettingsBase* FIKRetargetPreviewPropOp::GetSettings()
 {
 	return &Settings;
@@ -177,7 +215,74 @@ FName FIKRetargetPreviewPropOp::ApplyBoneMap(FName SourceBone)
 
 #if WITH_EDITOR
 // -----------------
-// All functionality is editor-only since at runtime 
+// All functionality is editor-only since at runtime op does nothing
+
+void FIKRetargetPreviewPropOp::UpdateAnimSeqPlayhead(UAnimInstance* SourceAnimInstance)
+{
+	// Get source tick records
+	TMap<FName, FAnimGroupInstance> SyncGroupTickRecords = SourceAnimInstance->GetSyncGroupMapRead();
+	TArray<FAnimTickRecord> UngroupedTickRecords = SourceAnimInstance->GetUngroupedActivePlayersRead();
+	
+	auto FindFirstTickRecord = [](const TArray<FAnimTickRecord>& ActivePlayers)
+		{
+			for (const FAnimTickRecord& ActivePlayer : ActivePlayers)
+			{
+				if (Cast<UAnimSequence>(ActivePlayer.SourceAsset) != nullptr)
+				{
+					// Used to determine which notifies should we considered "active".
+					check(ActivePlayer.DeltaTimeRecord->IsPreviousValid())
+					return &ActivePlayer;
+				}
+			}
+			
+			return static_cast<const FAnimTickRecord*>(nullptr);
+		};
+	
+	const FAnimTickRecord* FoundTickRecord = nullptr;
+	// Find all the anim sequences (with sync groups) we updated this tick.
+	for (const TTuple<FName, FAnimGroupInstance>& SyncGroupPair : SyncGroupTickRecords)
+	{
+		FoundTickRecord = FindFirstTickRecord(SyncGroupPair.Value.ActivePlayers);
+		if (FoundTickRecord)
+		{
+			break;
+		}
+	}
+	
+	if (!FoundTickRecord)
+	{
+		// Find all the anim sequences (no sync groups) we updated this tick.
+		FoundTickRecord = FindFirstTickRecord(UngroupedTickRecords);
+	}
+	
+	if (!FoundTickRecord)
+	{
+		CacheSourceCharacterAnim = nullptr;
+		SourceAnimPlayhead = 0.0f; 
+	}
+	else
+	{
+		CacheSourceCharacterAnim = Cast<UAnimSequenceBase>(FoundTickRecord->SourceAsset); 
+		SourceAnimPlayhead = FoundTickRecord->DeltaTimeRecord->GetPrevious() + FoundTickRecord->DeltaTimeRecord->Delta;
+	}
+}
+
+float FIKRetargetPreviewPropOp::ComputePropAnimPlayhead(float PropAnimLength, bool bLooping) const
+{
+	if (!CacheSourceCharacterAnim)
+	{
+		return 0.0f;
+	}
+	
+	if (!bLooping)
+	{
+		return FMath::Min(PropAnimLength, SourceAnimPlayhead);
+	}
+	
+	// NOTE: This will pop at the end of char sequence if SourceAnimLength is not an integer multiple of PropAnimLength
+	// TODO: Potentially compute a playrate to force correct looping behaviour? 
+	return FMath::Fmod(SourceAnimPlayhead, PropAnimLength);
+}
 
 bool FIKRetargetPreviewPropOp::GetMeshVertTris(FPropMeshData& OutMeshData, const FPreviewPropsData& PreviewProp)
 {
diff --git a/Engine/Plugins/Experimental/Animation/RelativeIKOp/Source/PreviewPropOp/Public/PreviewPropOp.h b/Engine/Plugins/Experimental/Animation/RelativeIKOp/Source/PreviewPropOp/Public/PreviewPropOp.h
--- a/Engine/Plugins/Experimental/Animation/RelativeIKOp/Source/PreviewPropOp/Public/PreviewPropOp.h
+++ b/Engine/Plugins/Experimental/Animation/RelativeIKOp/Source/PreviewPropOp/Public/PreviewPropOp.h
@@ -15,7 +15,10 @@
 
 struct FIKRigGoal;
 struct FKShapeElem;
+class UAnimInstance;
+class UAnimSequenceBase;
 class UPhysicsAsset;
+class USkeletalMeshComponent;
 
 struct FPropTfmInfo
 {
@@ -107,6 +110,10 @@ struct FIKRetargetPreviewPropOp : public FIKRetargetOpBase
 		const double InDeltaTime,
 		const TArray<FTransform>& InSourceGlobalPose,
 		TArray<FTransform>& OutTargetGlobalPose) override;
+	
+	UE_API virtual void AnimGraphPreUpdateMainThread(
+		USkeletalMeshComponent& SourceMeshComponent,
+		USkeletalMeshComponent& TargetMeshComponent) override;
 
 	UE_API virtual FIKRetargetOpSettingsBase* GetSettings() override;
 	
@@ -121,6 +128,9 @@ private:
 	FName ApplyBoneMap(FName SourceBone);
 
 #if WITH_EDITOR
+	void UpdateAnimSeqPlayhead(UAnimInstance* SourceAnimInstance);
+	float ComputePropAnimPlayhead(float PropAnimLength, bool bLooping) const;
+	
 	static bool GetMeshVertTris(FPropMeshData& OutMeshData, const FPreviewPropsData& PreviewProp);
 	
 public:
@@ -140,7 +150,10 @@ private:
 	TArray<FPropMeshData> PropMeshData;
 	
 	static UE_API FCriticalSection DebugDataMutex;
-
+	
+	// Cache playing anim-seq/montage for grabbing a playhead
+	TObjectPtr<const UAnimSequenceBase> CacheSourceCharacterAnim;
+	float SourceAnimPlayhead = 0.0f;
 	
 	TStrongObjectPtr<UMaterialInstanceDynamic> PropDefaultMaterial = nullptr;
 #endif
