BodyIntersectOp: Added hacky pole vector support for post-ik pushout using spheres on bone chains (e.g. at elbows)

#Synced-CL 49145592

diff --git a/Engine/Plugins/Experimental/Animation/RelativeIKOp/Source/BodyIntersectIKOp/Private/BodyIntersectIKOp.cpp b/Engine/Plugins/Experimental/Animation/RelativeIKOp/Source/BodyIntersectIKOp/Private/BodyIntersectIKOp.cpp
--- a/Engine/Plugins/Experimental/Animation/RelativeIKOp/Source/BodyIntersectIKOp/Private/BodyIntersectIKOp.cpp
+++ b/Engine/Plugins/Experimental/Animation/RelativeIKOp/Source/BodyIntersectIKOp/Private/BodyIntersectIKOp.cpp
@@ -62,13 +62,21 @@ bool FIKRetargetBodyIntersectIKOp::Initialize(
 	}
 
 	PhysicsAsset = Settings.TargetPhysicsAssetOverride;
-
 	if (!PhysicsAsset)
 	{
 		InLog.LogWarning( FText::Format(
 		LOCTEXT("MissingPhysicsAsset", "{0}: Intersect Physics Asset must be specified. "), FText::FromName(GetName())));
 		return false;
 	}
+	
+	const FRetargeterBoneChains& BoneChains = InProcessor.GetBoneChains();
+	const TArray<FResolvedBoneChain>* TargetChains = BoneChains.GetResolvedBoneChains(ERetargetSourceOrTarget::Target, ParentRigOp->Settings.IKRigAsset);
+	if (!TargetChains)
+	{
+		return false;
+	}
+	
+	InitPoleVectorIntersectors(*TargetChains, InTargetSkeleton);
 
 	UpdateCacheSkelInfo(InTargetSkeleton);
 
@@ -134,8 +142,8 @@ void FIKRetargetBodyIntersectIKOp::Run(
 		const int32 TargetBoneIdx = CacheTargetSkelIndices[GoalBoneName];
 		const FTransform TargetBoneTfm = OutTargetGlobalPose[TargetBoneIdx];
 
-		FKShapeElem* GoalShape = FindBodyShape(PhysicsAsset, IKGoal->BoneName);
-		double GoalRadius = CalcShapeSmallRadius(GoalShape); // * GoalSettings.SphereScale;
+		FKShapeElem* GoalShape = FPhysShapeUtils::FindBodyShape(PhysicsAsset, IKGoal->BoneName);
+		double GoalRadius = FPhysShapeUtils::CalcShapeSmallRadius(GoalShape); // * GoalSettings.SphereScale;
 		if (FMath::IsNearlyZero(GoalRadius))
 		{
 			continue;
@@ -163,14 +171,14 @@ void FIKRetargetBodyIntersectIKOp::Run(
 				
 			int32 BodyPoseIdx = CacheTargetSkelIndices[BodyName];
 			FTransform BodyTfm = OutTargetGlobalPose[BodyPoseIdx];
-			FKShapeElem* ShapeElem = FindBodyShape(PhysicsAsset, BodyName);
+			FKShapeElem* ShapeElem = FPhysShapeUtils::FindBodyShape(PhysicsAsset, BodyName);
 			if (!ShapeElem)
 			{
 				continue;
 			}
 
 			FVector DeltaDir = FVector::ZeroVector;
-			double Dist = FBodyIntersectHelper::CalcIntersectionPairDelta(GoalTfm, GoalShape, BodyTfm, ShapeElem, DeltaDir);
+			double Dist = FBodyIntersectUtils::CalcIntersectionPairDelta(GoalTfm, GoalShape, BodyTfm, ShapeElem, DeltaDir);
 			if (Dist > 0.0)
 			{
 				SphereOffsetGoal += Dist * DeltaDir;
@@ -214,6 +222,272 @@ void FIKRetargetBodyIntersectIKOp::Run(
 #endif
 }
 
+void FIKRetargetBodyIntersectIKOp::RunAfterParent(
+	FIKRetargetProcessor& InProcessor,
+	const double InDeltaTime,
+	const TArray<FTransform>& InSourceGlobalPose,
+	TArray<FTransform>& OutTargetGlobalPose)
+{
+	if (InProcessor.IsIKForcedOff())
+	{
+		return;
+	}
+
+	if (!bIsInitialized)
+	{
+		return;
+	}
+	
+	if (!Settings.bEnablePoleVectorIntersect)
+	{
+		return;
+	}
+	
+	FDebugBodyIntersectDrawInfo LocalDebugInfo;
+	
+	const FRetargetSkeleton& TargetSkeleton = InProcessor.GetTargetSkeleton();
+	for (const FPoleVectorChainIntersector& PVI : PoleVectorChainIntersectors)
+	{
+		PVI.ApplyIntersectRotation(TargetSkeleton, OutTargetGlobalPose, LocalDebugInfo);
+	}
+	
+#if WITH_EDITOR
+	if (Settings.bDebugDraw)
+	{
+		FScopeLock ScopeLock(&DebugDataMutex);
+		DebugDrawInfo.PoleIntersectTfms = LocalDebugInfo.PoleIntersectTfms;
+	}
+#endif //WITH_EDITOR
+}
+
+void FPoleVectorChainIntersector::Initialize(
+	const FIKRetargetBodyIntersectIKOpSettings* InSettings,
+	const TArray<int32>& InBodySettingsIdx,
+	const FResolvedBoneChain* InTargetChain,
+	const FTargetSkeleton& InTargetSkeleton)
+{
+	Settings = InSettings;
+	PhysicsAsset = InSettings->TargetPhysicsAssetOverride;
+	TargetBoneChain = InTargetChain;
+	IntersectBonesIdx.Reset();
+	PoleBodiesSettingsIdx.Reset();
+	PoleBonesIdx.Reset();
+	AllChildrenWithinChain.Reset();
+	
+	for (int32 SettingsIdx : InBodySettingsIdx)
+	{
+		const FIKPoleVectorIntersectSettings& PvIsectSetting = Settings->PoleVectorIntersectSettings[SettingsIdx];
+		int32 BoneIdx = InTargetSkeleton.FindBoneIndexByName(PvIsectSetting.BodyName);
+		if (BoneIdx == INDEX_NONE)
+		{
+			continue;
+		}
+		
+		PoleBodiesSettingsIdx.Add(SettingsIdx);
+		PoleBonesIdx.Add(BoneIdx);
+	}
+	
+	for (const FName IntersectBone : Settings->IntersectBodies)
+	{
+		// TODO: Normalize code to remove FName->Index maps and just keep struct of bones/indices
+		int32 BoneIdx = InTargetSkeleton.FindBoneIndexByName(IntersectBone);
+		IntersectBonesIdx.Add(BoneIdx);
+	}
+	
+	// cache indices of bones in chain (and their children recursively) excluding children of the end bone
+	TArray<int32> AllChildrenIndices;
+	InTargetSkeleton.GetChildrenIndicesRecursive(TargetBoneChain->BoneIndices[0],AllChildrenIndices);
+	TArray<int32> ChildrenOfEndIndices = {TargetBoneChain->BoneIndices.Last()};
+	InTargetSkeleton.GetChildrenIndicesRecursive(TargetBoneChain->BoneIndices.Last(),ChildrenOfEndIndices);
+	
+	for (int32 ChildIndex : AllChildrenIndices)
+	{
+		if (ChildrenOfEndIndices.Contains(ChildIndex))
+		{
+			continue;
+		}
+		AllChildrenWithinChain.Add(ChildIndex);
+	}
+}
+
+void FPoleVectorChainIntersector::ApplyIntersectRotation(
+	const FRetargetSkeleton& TargetSkeleton,
+	TArray<FTransform>& OutTargetGlobalPose,
+	FDebugBodyIntersectDrawInfo& DebugInfo) const
+{
+	const FVector TargetChainAxisNorm = GetChainAxisNormalized(TargetBoneChain->BoneIndices, OutTargetGlobalPose);
+	const FVector ChainOriginLoc = OutTargetGlobalPose[TargetBoneChain->BoneIndices[0]].GetLocation();
+	
+	// Get all child local bone transforms
+	TArray<FTransform> AllChildrenInputLocalTransforms = TargetSkeleton.GetLocalTransformsOfMultipleBones(AllChildrenWithinChain, OutTargetGlobalPose);
+	
+	double TotalWeight = 0.0;
+	double TotalAngle = 0.0;
+	for (int32 Index = 0; Index < PoleBonesIdx.Num(); ++Index)
+	{
+		int32 BoneIdx = PoleBonesIdx[Index];
+		int32 SettingsIdx = PoleBodiesSettingsIdx[Index];
+		
+		if (!Settings->PoleVectorIntersectSettings.IsValidIndex(SettingsIdx))
+		{
+			continue;
+		}
+		
+		const FTransform TargetBoneTfm = OutTargetGlobalPose[BoneIdx];
+		
+		// TODO: Should we copy these settings into the struct?
+		const FIKPoleVectorIntersectSettings& BodySettings = Settings->PoleVectorIntersectSettings[SettingsIdx];
+
+		FKShapeElem* BodyShape = FPhysShapeUtils::FindBodyShape(PhysicsAsset, BodySettings.BodyName);
+		double BodyRadius = FPhysShapeUtils::CalcShapeSmallRadius(BodyShape) * BodySettings.SphereScale;
+		if (FMath::IsNearlyZero(BodyRadius))
+		{
+			continue;
+		}
+		
+		FTransform LocalTfm(BodySettings.SphereOffset);
+		FTransform PoleTfm = LocalTfm * TargetBoneTfm;
+		FVector SphereLoc = PoleTfm.GetLocation();
+		FVector SphereOriginVec = SphereLoc - ChainOriginLoc;
+		FVector SpherePoleVec = FVector::VectorPlaneProject(SphereOriginVec, TargetChainAxisNorm);
+		
+		// TODO: Potentially blend/early-out when pole vec is much smaller than sphere radius
+		
+		FKSphereElem PoleBodySphere;
+		PoleBodySphere.Radius = static_cast<float>(BodyRadius);
+		PoleBodySphere.Center = BodySettings.SphereOffset;
+		
+		FVector IsectDelta = GetIntersectDelta(PoleTfm, &PoleBodySphere, OutTargetGlobalPose, DebugInfo);
+		FVector PoleDelta = FVector::VectorPlaneProject(IsectDelta, TargetChainAxisNorm);
+		if (PoleDelta.IsNearlyZero())
+		{
+			continue;
+		}
+		
+		FVector DesiredLoc = SpherePoleVec + PoleDelta;
+		
+		FVector PoleNorm = SpherePoleVec.GetSafeNormal();
+		FVector RotDir = FVector::CrossProduct(TargetChainAxisNorm, PoleNorm);
+		
+		FVector DesiredNorm = DesiredLoc.GetSafeNormal();
+		double Angle = FMath::Atan2(FVector::DotProduct(DesiredNorm, RotDir), FVector::DotProduct(DesiredNorm,PoleNorm));
+		
+		TotalAngle += Angle;
+		TotalWeight += 1.0;
+	}
+	
+	if (!FMath::IsNearlyZero(TotalWeight))
+	{
+		FQuat Rotation(TargetChainAxisNorm, TotalAngle / TotalWeight);
+	
+		// rotate the base of the chain to match the pole vectors
+		FTransform& BaseOfChain = OutTargetGlobalPose[TargetBoneChain->BoneIndices[0]];
+		BaseOfChain.SetRotation(Rotation * BaseOfChain.GetRotation());
+	
+		// Update changes to children of the chain
+		TargetSkeleton.UpdateGlobalTransformsOfMultipleBones(AllChildrenWithinChain, AllChildrenInputLocalTransforms,OutTargetGlobalPose);
+	}
+	
+#if WITH_EDITOR
+	if (Settings->bDebugDraw)
+	{
+		for (int32 Index = 0; Index < PoleBonesIdx.Num(); ++Index)
+		{
+			int32 BoneIdx = PoleBonesIdx[Index];
+			int32 SettingsIdx = PoleBodiesSettingsIdx[Index];
+		
+			if (!Settings->PoleVectorIntersectSettings.IsValidIndex(SettingsIdx))
+			{
+				continue;
+			}
+		
+			const FTransform TargetBoneTfm = OutTargetGlobalPose[BoneIdx];
+			
+			const FIKPoleVectorIntersectSettings& BodySettings = Settings->PoleVectorIntersectSettings[SettingsIdx];
+
+			FKShapeElem* BodyShape = FPhysShapeUtils::FindBodyShape(PhysicsAsset, BodySettings.BodyName);
+			double BodyRadius = FPhysShapeUtils::CalcShapeSmallRadius(BodyShape) * BodySettings.SphereScale;
+			if (FMath::IsNearlyZero(BodyRadius))
+			{
+				continue;
+			}
+		
+			FTransform LocalTfm(BodySettings.SphereOffset);
+			DebugInfo.PoleIntersectTfms.Emplace(static_cast<float>(BodyRadius), LocalTfm * TargetBoneTfm);
+		}
+	}
+#endif //WITH_EDITOR
+}
+
+FVector FPoleVectorChainIntersector::GetIntersectDelta(
+	const FTransform& PoleTfm,
+	FKSphereElem* PoleSphere,
+	const TArray<FTransform>& TargetGlobalPose,
+	FDebugBodyIntersectDrawInfo& DebugInfo) const
+{
+	FVector TotalDelta = FVector::ZeroVector;
+	for (int Index = 0; Index < Settings->IntersectBodies.Num(); ++Index)
+	{
+		if ( !IntersectBonesIdx.IsValidIndex(Index) )
+		{
+			continue;
+		}
+		
+		int32 BoneIdx = IntersectBonesIdx[Index];
+		if (BoneIdx == INDEX_NONE)
+		{
+			continue;
+		}
+		
+		FName BodyName = Settings->IntersectBodies[Index];
+		if (PhysicsAsset->FindBodyIndex(BodyName) == INDEX_NONE)
+		{
+			continue;
+		}
+		
+		FTransform BodyTfm = TargetGlobalPose[BoneIdx];
+		FKShapeElem* ShapeElem = FPhysShapeUtils::FindBodyShape(PhysicsAsset, BodyName);
+		if (!ShapeElem)
+		{
+			continue;
+		}
+
+		FVector DeltaDir = FVector::ZeroVector;
+		double Dist = FBodyIntersectUtils::CalcIntersectionPairDelta(PoleTfm, PoleSphere, BodyTfm, ShapeElem, DeltaDir);
+		if (Dist > 0.0)
+		{
+			TotalDelta += DeltaDir * Dist;
+#if WITH_EDITOR
+			// TODO: Add intersection detected code here
+#endif //WITH_EDITOR
+		}
+	}
+	
+	return TotalDelta;
+}
+
+FVector FPoleVectorChainIntersector::CalculatePoleVector(
+	const EAxis::Type& PoleAxis,
+	const TArray<int32>& BoneIndices,
+	const TArray<FTransform>& GlobalPose) const
+{
+	check(!BoneIndices.IsEmpty())
+
+	const FVector ChainNormal = GetChainAxisNormalized(BoneIndices, GlobalPose);
+	const FVector UnitPoleAxis = GlobalPose[BoneIndices[0]].GetUnitAxis(PoleAxis);
+	const FVector PoleVector = FVector::VectorPlaneProject(UnitPoleAxis, ChainNormal);
+	return PoleVector.GetSafeNormal();
+}
+
+FVector FPoleVectorChainIntersector::GetChainAxisNormalized(
+	const TArray<int32>& BoneIndices,
+	const TArray<FTransform>& GlobalPose) const
+{
+	const FVector ChainOrigin = GlobalPose[BoneIndices[0]].GetLocation();
+	const FVector ChainAxis = GlobalPose[BoneIndices.Last()].GetLocation() - ChainOrigin;
+	return ChainAxis.GetSafeNormal();
+}
+
 FVector FIKRetargetBodyIntersectIKOp::GoalLocBlendCompSpace(const FIKRigGoal* Goal, const FTransform& BoneTfm) const
 {
 	switch (Goal->PositionSpace)
@@ -264,78 +538,48 @@ void FIKRetargetBodyIntersectIKOp::SetGoalPosFromCompSpace(FIKRigGoal* Goal, con
 	}
 }
 
-double FIKRetargetBodyIntersectIKOp::CalcShapeSmallRadius(FKShapeElem* ShapeElem)
+void FIKRetargetBodyIntersectIKOp::InitPoleVectorIntersectors(
+	const TArray<FResolvedBoneChain>& TargetChains,
+	const FTargetSkeleton& InTargetSkeleton)
 {
-	// Return 1.0 as default so scaling can set "size" of sphere
-	const double DefaultRadius = 1.0;
-	if (!ShapeElem)
+	PoleVectorChainIntersectors.Reset();
+	
+	if (!Settings.bEnablePoleVectorIntersect)
 	{
-		return DefaultRadius;
+		return;
 	}
 	
-	switch (ShapeElem->GetShapeType())
+	// TODO: Can remove a lot of the checks/nesting in the loop if we use fix up code to user per-chain settings/info instead
+	TMap<const FResolvedBoneChain*, TArray<int32>> ChainSetupMap;
+	for (int32 SettingsIdx = 0; SettingsIdx < Settings.PoleVectorIntersectSettings.Num(); ++SettingsIdx)
 	{
-	case EAggCollisionShape::Sphere:
-		{
-			FKSphereElem* SphereElem = static_cast<FKSphereElem*>(ShapeElem);
-			return SphereElem->Radius;
-		}
-
-	case EAggCollisionShape::Sphyl:
-		{
-			FKSphylElem* CapsuleElem = static_cast<FKSphylElem*>(ShapeElem);
-			return CapsuleElem->Radius;
-		}
-	case EAggCollisionShape::Box:
+		const FIKPoleVectorIntersectSettings& PoleVecInfo = Settings.PoleVectorIntersectSettings[SettingsIdx];
+		int32 BoneIdx = InTargetSkeleton.FindBoneIndexByName(PoleVecInfo.BodyName);
+		if (BoneIdx == INDEX_NONE)
 		{
-			FKBoxElem* BoxElem = static_cast<FKBoxElem*>(ShapeElem);
-			return 0.5 * FMath::Min(FMath::Min(BoxElem->X, BoxElem->Y), BoxElem->Z);
-		}
-		// case EAggCollisionShape::Convex:
-		// 	{
-		// 		FKConvexElem* ConvElem = static_cast<FKConvexElem*>(ShapeElem);
-		// 		return 0.5 * ConvElem->ElemBox.GetSize().GetMin();
-		// 	}
-
-	default:
-		{
-			return DefaultRadius;
-		}
-	}
-}
-
-double FIKRetargetBodyIntersectIKOp::CalcShapeAvgRadius(FKShapeElem* ShapeElem)
-{
-	switch (ShapeElem->GetShapeType())
-	{
-	case EAggCollisionShape::Sphere:
-		{
-			FKSphereElem* SphereElem = static_cast<FKSphereElem*>(ShapeElem);
-			return SphereElem->Radius;
-		}
-
-	case EAggCollisionShape::Sphyl:
-		{
-			FKSphylElem* CapsuleElem = static_cast<FKSphylElem*>(ShapeElem);
-			return CapsuleElem->Radius + 0.25*CapsuleElem->Length;
+			continue;
 		}
-
-	default:
+		
+		for (const FResolvedBoneChain& Chain : TargetChains)
 		{
-			return 0.0;
+			if (!Chain.BoneIndices.Contains(BoneIdx))
+			{
+				continue;
+			}
+			
+			if (!ChainSetupMap.Contains(&Chain))
+			{
+				ChainSetupMap.Emplace(&Chain, {});
+			}
+			ChainSetupMap[&Chain].Add(SettingsIdx);
 		}
 	}
-}
-
-FKShapeElem* FIKRetargetBodyIntersectIKOp::FindBodyShape(const UPhysicsAsset* PhysAsset, FName BoneName)
-{
-	int32 BodyIdx = PhysAsset->FindBodyIndex(BoneName);
-	if (BodyIdx == INDEX_NONE)
+	
+	for (const TPair<const FResolvedBoneChain*, TArray<int32>>& ChainSetupPair : ChainSetupMap)
 	{
-		UE_LOG(LogAnimation, Warning, TEXT("No body index found: %s"), *BoneName.ToString());
-		return nullptr;
+		FPoleVectorChainIntersector& NewPVI = PoleVectorChainIntersectors.AddDefaulted_GetRef();
+		NewPVI.Initialize(&Settings, ChainSetupPair.Value, ChainSetupPair.Key, InTargetSkeleton);
 	}
-	return PhysAsset->SkeletalBodySetups[BodyIdx]->AggGeom.GetElement(0);
 }
 
 void FIKRetargetBodyIntersectIKOp::UpdateTargetBoneMap(FName TargetBoneName)
@@ -406,7 +650,7 @@ void FIKRetargetBodyIntersectIKOp::DebugDraw(
 	{
 		const TTuple<FName, FTransform>& IntersectBody = DebugDrawInfo.TargetIntersectTfms[i];
 		const FName BodyName = IntersectBody.Get<0>();
-		const FKShapeElem* ShapeElem = FindBodyShape(Settings.TargetPhysicsAssetOverride, BodyName);
+		const FKShapeElem* ShapeElem = FPhysShapeUtils::FindBodyShape(Settings.TargetPhysicsAssetOverride, BodyName);
 		if (!ShapeElem)
 		{
 			continue;
@@ -427,7 +671,7 @@ void FIKRetargetBodyIntersectIKOp::DebugDraw(
 	for (const TTuple<FName, FTransform>& GoalIntersectShape : DebugDrawInfo.GoalIntersectTfms)
 	{
 		const FName BodyName = GoalIntersectShape.Get<0>();
-		const FKShapeElem* ShapeElem = FindBodyShape(Settings.TargetPhysicsAssetOverride, BodyName);
+		const FKShapeElem* ShapeElem = FPhysShapeUtils::FindBodyShape(Settings.TargetPhysicsAssetOverride, BodyName);
 		if (!ShapeElem)
 		{
 			continue;
@@ -439,6 +683,18 @@ void FIKRetargetBodyIntersectIKOp::DebugDraw(
 		FTransform CompTfm = GoalShapeTfm * InComponentTransform;
 		DebugDrawPhysBody(InPDI, CompTfm, ShapeScale, ShapeElem, FLinearColor::Red);
 	}
+	
+	for (const TTuple<float, FTransform>& PoleIntersect : DebugDrawInfo.PoleIntersectTfms)
+	{
+		FKSphereElem SphereElem;
+		SphereElem.Radius = PoleIntersect.Get<0>();
+		
+		FTransform PoleShapeTfm = PoleIntersect.Get<1>();
+		FVector ShapeScale = PoleShapeTfm.GetScale3D()*InComponentScale;
+		PoleShapeTfm.RemoveScaling();
+		FTransform CompTfm = PoleShapeTfm * InComponentTransform;
+		DebugDrawPhysBody(InPDI, CompTfm, ShapeScale, &SphereElem, FLinearColor::Red);
+	}
 }
 
 void FIKRetargetBodyIntersectIKOp::DebugDrawPhysBody(FPrimitiveDrawInterface* InPDI, const FTransform& ParentTransform, const FVector& Scale, const FKShapeElem* ShapeElem, const FLinearColor& Color) const
diff --git a/Engine/Plugins/Experimental/Animation/RelativeIKOp/Source/BodyIntersectIKOp/Private/PhysicsBodyHelpers.cpp b/Engine/Plugins/Experimental/Animation/RelativeIKOp/Source/BodyIntersectIKOp/Private/PhysicsBodyHelpers.cpp
--- a/Engine/Plugins/Experimental/Animation/RelativeIKOp/Source/BodyIntersectIKOp/Private/PhysicsBodyHelpers.cpp
+++ b/Engine/Plugins/Experimental/Animation/RelativeIKOp/Source/BodyIntersectIKOp/Private/PhysicsBodyHelpers.cpp
@@ -5,7 +5,7 @@
 #include "PhysicsEngine/PhysicsAsset.h"
 #include "PhysicsEngine/SkeletalBodySetup.h"
 
-double FBodyIntersectHelper::CalcIntersectionPairDelta(const FTransform& GoalShapeTfm, FKShapeElem* GoalShapeElem, const FTransform& TargetShapeTfm, FKShapeElem* TargetShapeElem, FVector& OutDeltaDir)
+double FBodyIntersectUtils::CalcIntersectionPairDelta(const FTransform& GoalShapeTfm, FKShapeElem* GoalShapeElem, const FTransform& TargetShapeTfm, FKShapeElem* TargetShapeElem, FVector& OutDeltaDir)
 {
 	if (!GoalShapeElem || !TargetShapeElem)
 	{
@@ -106,7 +106,7 @@ double FBodyIntersectHelper::CalcIntersectionPairDelta(const FTransform& GoalSha
 	return Dist;
 }
 
-double FBodyIntersectHelper::SegmentDistToAABB(
+double FBodyIntersectUtils::SegmentDistToAABB(
 	const FVector& A1, const FVector& A2,
 	const FVector& BoxMin, const FVector& BoxMax,
 	FVector& OutSeg, FVector& OutBox)
@@ -167,7 +167,7 @@ double FBodyIntersectHelper::SegmentDistToAABB(
 	return FMath::Sqrt(bestD2);
 }
 
-double FBodyIntersectHelper::CalcSphereSphereShapeDelta(FVector& OutDeltaDir, const FTransform& ShapeTfm1, const FKSphereElem* ShapeElem1, const FTransform& ShapeTfm2, const FKSphereElem* ShapeElem2)
+double FBodyIntersectUtils::CalcSphereSphereShapeDelta(FVector& OutDeltaDir, const FTransform& ShapeTfm1, const FKSphereElem* ShapeElem1, const FTransform& ShapeTfm2, const FKSphereElem* ShapeElem2)
 {
 	const double SphereRadius1 = ShapeTfm1.GetScale3D().GetAbsMin() * ShapeElem1->Radius;
 	const FVector WorldCenter1 = ShapeTfm1.TransformPositionNoScale(ShapeElem1->Center);
@@ -182,7 +182,7 @@ double FBodyIntersectHelper::CalcSphereSphereShapeDelta(FVector& OutDeltaDir, co
 	return DistToEdge;
 }
 
-double FBodyIntersectHelper::CalcSphereBoxShapeDelta(FVector& OutDeltaDir, const FTransform& ShapeTfm1, const FKSphereElem* ShapeElem1, const FTransform& ShapeTfm2, const FKBoxElem* ShapeElem2)
+double FBodyIntersectUtils::CalcSphereBoxShapeDelta(FVector& OutDeltaDir, const FTransform& ShapeTfm1, const FKSphereElem* ShapeElem1, const FTransform& ShapeTfm2, const FKBoxElem* ShapeElem2)
 {
 	const double SphereRadius1 = ShapeTfm1.GetScale3D().GetAbsMin() * ShapeElem1->Radius;
 	const FVector WorldCenter1 = ShapeTfm1.TransformPositionNoScale(ShapeElem1->Center);
@@ -228,7 +228,7 @@ double FBodyIntersectHelper::CalcSphereBoxShapeDelta(FVector& OutDeltaDir, const
 	return DistToEdge;
 }
 
-double FBodyIntersectHelper::CalcSphereCapsuleShapeDelta(FVector& OutDeltaDir, const FTransform& ShapeTfm1, const FKSphereElem* ShapeElem1, const FTransform& ShapeTfm2, const FKSphylElem* ShapeElem2)
+double FBodyIntersectUtils::CalcSphereCapsuleShapeDelta(FVector& OutDeltaDir, const FTransform& ShapeTfm1, const FKSphereElem* ShapeElem1, const FTransform& ShapeTfm2, const FKSphylElem* ShapeElem2)
 {
 	const double SphereRadius1 = ShapeTfm1.GetScale3D().GetAbsMin() * ShapeElem1->Radius;
 
@@ -252,14 +252,14 @@ double FBodyIntersectHelper::CalcSphereCapsuleShapeDelta(FVector& OutDeltaDir, c
 	return DistToEdge;
 }
 
-double FBodyIntersectHelper::CalcBoxBoxShapeDelta(FVector& OutDeltaDir, const FTransform& ShapeTfm1, const FKBoxElem* ShapeElem1, const FTransform& ShapeTfm2, const FKBoxElem* ShapeElem2)
+double FBodyIntersectUtils::CalcBoxBoxShapeDelta(FVector& OutDeltaDir, const FTransform& ShapeTfm1, const FKBoxElem* ShapeElem1, const FTransform& ShapeTfm2, const FKBoxElem* ShapeElem2)
 {
 	// TODO: Warning for no implementation
 	OutDeltaDir = FVector::ZeroVector;
 	return 0;
 }
 
-double FBodyIntersectHelper::CalcBoxCapsuleShapeDelta(FVector& OutDeltaDir, const FTransform& ShapeTfm1, const FKBoxElem* ShapeElem1, const FTransform& ShapeTfm2, const FKSphylElem* ShapeElem2)
+double FBodyIntersectUtils::CalcBoxCapsuleShapeDelta(FVector& OutDeltaDir, const FTransform& ShapeTfm1, const FKBoxElem* ShapeElem1, const FTransform& ShapeTfm2, const FKSphylElem* ShapeElem2)
 {
 	FTransform ShapeTfm1NoScale = ShapeTfm1;
 	ShapeTfm1NoScale.RemoveScaling();
@@ -341,7 +341,7 @@ double FBodyIntersectHelper::CalcBoxCapsuleShapeDelta(FVector& OutDeltaDir, cons
 	return DistToEdge;
 }
 
-double FBodyIntersectHelper::CalcCapsuleCapsuleShapeDelta(FVector& OutDeltaDir, const FTransform& ShapeTfm1, const FKSphylElem* ShapeElem1, const FTransform& ShapeTfm2, const FKSphylElem* ShapeElem2)
+double FBodyIntersectUtils::CalcCapsuleCapsuleShapeDelta(FVector& OutDeltaDir, const FTransform& ShapeTfm1, const FKSphylElem* ShapeElem1, const FTransform& ShapeTfm2, const FKSphylElem* ShapeElem2)
 {
 	FTransform ShapeTfm1NoScale = ShapeTfm1;
 	ShapeTfm1NoScale.RemoveScaling();
@@ -372,7 +372,7 @@ double FBodyIntersectHelper::CalcCapsuleCapsuleShapeDelta(FVector& OutDeltaDir,
 	return DistToEdge;
 }
 
-EBodyIntersectPairType FBodyIntersectHelper::CalcBodyIntersectPairType(const FKShapeElem* ShapeElem1, const FKShapeElem* ShapeElem2, bool &NeedInverse)
+EBodyIntersectPairType FBodyIntersectUtils::CalcBodyIntersectPairType(const FKShapeElem* ShapeElem1, const FKShapeElem* ShapeElem2, bool &NeedInverse)
 {
 	uint8 TypeIndex1 = static_cast<uint8>(ShapeElem1->GetShapeType());
 	uint8 TypeIndex2 = static_cast<uint8>(ShapeElem2->GetShapeType());
@@ -390,3 +390,77 @@ EBodyIntersectPairType FBodyIntersectHelper::CalcBodyIntersectPairType(const FKS
 	}
 	return EBodyIntersectPairType::Unknown;
 }
+
+double FPhysShapeUtils::CalcShapeSmallRadius(FKShapeElem* ShapeElem)
+{
+	// Return 1.0 as default so scaling can set "size" of sphere
+	const double DefaultRadius = 1.0;
+	if (!ShapeElem)
+	{
+		return DefaultRadius;
+	}
+	
+	switch (ShapeElem->GetShapeType())
+	{
+	case EAggCollisionShape::Sphere:
+		{
+			FKSphereElem* SphereElem = static_cast<FKSphereElem*>(ShapeElem);
+			return SphereElem->Radius;
+		}
+
+	case EAggCollisionShape::Sphyl:
+		{
+			FKSphylElem* CapsuleElem = static_cast<FKSphylElem*>(ShapeElem);
+			return CapsuleElem->Radius;
+		}
+	case EAggCollisionShape::Box:
+		{
+			FKBoxElem* BoxElem = static_cast<FKBoxElem*>(ShapeElem);
+			return 0.5 * FMath::Min(FMath::Min(BoxElem->X, BoxElem->Y), BoxElem->Z);
+		}
+		// case EAggCollisionShape::Convex:
+		// 	{
+		// 		FKConvexElem* ConvElem = static_cast<FKConvexElem*>(ShapeElem);
+		// 		return 0.5 * ConvElem->ElemBox.GetSize().GetMin();
+		// 	}
+
+	default:
+		{
+			return DefaultRadius;
+		}
+	}
+}
+
+double FPhysShapeUtils::CalcShapeAvgRadius(FKShapeElem* ShapeElem)
+{
+	switch (ShapeElem->GetShapeType())
+	{
+	case EAggCollisionShape::Sphere:
+		{
+			FKSphereElem* SphereElem = static_cast<FKSphereElem*>(ShapeElem);
+			return SphereElem->Radius;
+		}
+
+	case EAggCollisionShape::Sphyl:
+		{
+			FKSphylElem* CapsuleElem = static_cast<FKSphylElem*>(ShapeElem);
+			return CapsuleElem->Radius + 0.25*CapsuleElem->Length;
+		}
+
+	default:
+		{
+			return 0.0;
+		}
+	}
+}
+
+FKShapeElem* FPhysShapeUtils::FindBodyShape(const UPhysicsAsset* PhysAsset, FName BoneName)
+{
+	int32 BodyIdx = PhysAsset->FindBodyIndex(BoneName);
+	if (BodyIdx == INDEX_NONE)
+	{
+		UE_LOG(LogAnimation, Warning, TEXT("No body index found: %s"), *BoneName.ToString());
+		return nullptr;
+	}
+	return PhysAsset->SkeletalBodySetups[BodyIdx]->AggGeom.GetElement(0);
+}
diff --git a/Engine/Plugins/Experimental/Animation/RelativeIKOp/Source/BodyIntersectIKOp/Private/PhysicsBodyHelpers.h b/Engine/Plugins/Experimental/Animation/RelativeIKOp/Source/BodyIntersectIKOp/Private/PhysicsBodyHelpers.h
--- a/Engine/Plugins/Experimental/Animation/RelativeIKOp/Source/BodyIntersectIKOp/Private/PhysicsBodyHelpers.h
+++ b/Engine/Plugins/Experimental/Animation/RelativeIKOp/Source/BodyIntersectIKOp/Private/PhysicsBodyHelpers.h
@@ -9,6 +9,7 @@ struct FKShapeElem;
 struct FKSphereElem;
 struct FKBoxElem;
 struct FKSphylElem;
+class UPhysicsAsset;
 
 enum class EBodyIntersectPairType : uint8 
 {
@@ -25,7 +26,7 @@ enum class EBodyIntersectPairType : uint8
 /**
  * Helper static class providing physics body intersection delta functions
  */
-class FBodyIntersectHelper
+class FBodyIntersectUtils
 {
 public:
 	static double CalcIntersectionPairDelta(const FTransform& GoalShapeTfm, FKShapeElem* GoalShapeElem, const FTransform& TargetShapeTfm, FKShapeElem* TargetShapeElem, FVector& OutDeltaDir);
@@ -42,3 +43,12 @@ private:
 
 	static double SegmentDistToAABB(const FVector& A, const FVector& B, const FVector& BoxMin, const FVector& BoxMax, FVector& OutSeg, FVector& OutBox);
 };
+
+class FPhysShapeUtils
+{
+public:
+	static double CalcShapeSmallRadius(FKShapeElem* ShapeElem);
+	static double CalcShapeAvgRadius(FKShapeElem* ShapeElem);
+
+	static FKShapeElem* FindBodyShape(const UPhysicsAsset* PhysAsset, FName BoneName);
+};
diff --git a/Engine/Plugins/Experimental/Animation/RelativeIKOp/Source/BodyIntersectIKOp/Public/BodyIntersectIKOp.h b/Engine/Plugins/Experimental/Animation/RelativeIKOp/Source/BodyIntersectIKOp/Public/BodyIntersectIKOp.h
--- a/Engine/Plugins/Experimental/Animation/RelativeIKOp/Source/BodyIntersectIKOp/Public/BodyIntersectIKOp.h
+++ b/Engine/Plugins/Experimental/Animation/RelativeIKOp/Source/BodyIntersectIKOp/Public/BodyIntersectIKOp.h
@@ -15,6 +15,7 @@ struct FKShapeElem;
 struct FKSphereElem;
 struct FKBoxElem;
 struct FKSphylElem;
+struct FResolvedBoneChain;
 class UPhysicsAsset;
 class USkeletalMesh;
 
@@ -25,6 +26,8 @@ struct FDebugBodyIntersectDrawInfo
 	TArray<TTuple<FName,FTransform>> TargetIntersectTfms;
 	TArray<bool> TargetIntersectDetected;
 	TArray<TTuple<FName,FTransform>> GoalIntersectTfms;
+	
+	TArray<TTuple<float,FTransform>> PoleIntersectTfms;
 };
 
 USTRUCT(BlueprintType, meta = (DisplayName = "Goal Intersect Settings"))
@@ -43,6 +46,24 @@ struct FIKGoalIntersectShapeSettings
 	FVector GoalShapeOffset = FVector::ZeroVector;
 };
 
+USTRUCT(BlueprintType, meta = (DisplayName = "Pole Vector Settings"))
+struct FIKPoleVectorIntersectSettings
+{
+	GENERATED_BODY()
+
+	// The body along ik chain to check intersections and apply pole vector rotation
+	UPROPERTY(EditAnywhere, BlueprintReadWrite, Category = "Pole Vector Intersect Settings", meta=(ReinitializeOnEdit))
+	FName BodyName = NAME_None;
+	
+	// Sphere size scale (base radius is body min radius or 1.0) for intersection checks
+	UPROPERTY(EditAnywhere, BlueprintReadWrite, Category = "Pole Vector Intersect Settings", meta = (ClampMin="0.0"))
+	double SphereScale = 1.0;
+	
+	// Local sphere offset from body bone
+	UPROPERTY(EditAnywhere, BlueprintReadWrite, Category = "Pole Vector Intersect Settings")
+	FVector SphereOffset = FVector::ZeroVector;
+};
+
 USTRUCT(BlueprintType, meta = (DisplayName = "Body Intersect IK Settings"))
 struct FIKRetargetBodyIntersectIKOpSettings : public FIKRetargetOpSettingsBase
 {
@@ -64,6 +85,55 @@ struct FIKRetargetBodyIntersectIKOpSettings : public FIKRetargetOpSettingsBase
 	// **NOTE: Only Sphere, Capsule and Box bodies currently supported**
 	UPROPERTY(EditAnywhere, BlueprintReadWrite, Category = "Intersection", meta=(ReinitializeOnEdit))
 	TArray<FName> IntersectBodies;
+	
+	// Check bodies in Pole Vector Intersect Settings for intersection and try to fixup using pole vector rotation
+	UPROPERTY(EditAnywhere, BlueprintReadWrite, Category = "Pole Vector", meta=(ReinitializeOnEdit))
+	bool bEnablePoleVectorIntersect = true;
+	
+	// Body settings for pole vector intersection fixup
+	UPROPERTY(EditAnywhere, BlueprintReadWrite, Category = "Pole Vector", meta=(ReinitializeOnEdit))
+	TArray<FIKPoleVectorIntersectSettings> PoleVectorIntersectSettings;
+};
+
+// Per-chain pole vector intersection functionality
+struct FPoleVectorChainIntersector
+{
+	void Initialize(
+		const FIKRetargetBodyIntersectIKOpSettings* InSettings,
+		const TArray<int32>& InBodySettingsIdx,
+		const FResolvedBoneChain* InTargetChain,
+		const FTargetSkeleton& InTargetSkeleton);
+	
+	void ApplyIntersectRotation(
+		const FRetargetSkeleton& TargetSkeleton,
+		TArray<FTransform>& OutTargetGlobalPose,
+		FDebugBodyIntersectDrawInfo& DebugInfo) const;
+	
+	FVector GetIntersectDelta(
+		const FTransform& PoleTfm,
+		FKSphereElem* PoleSphere,
+		const TArray<FTransform>& TargetGlobalPose,
+		FDebugBodyIntersectDrawInfo& DebugInfo) const;
+	
+	FVector CalculatePoleVector(
+		const EAxis::Type& PoleAxis,
+		const TArray<int32>& BoneIndices,
+		const TArray<FTransform>& GlobalPose) const;
+	
+	FVector GetChainAxisNormalized(
+		const TArray<int32>& BoneIndices,
+		const TArray<FTransform>& GlobalPose) const;
+	
+	TObjectPtr<const UPhysicsAsset> PhysicsAsset;
+	
+	const FIKRetargetBodyIntersectIKOpSettings* Settings;
+	const FResolvedBoneChain* TargetBoneChain;
+	
+	TArray<int32> PoleBodiesSettingsIdx;
+	TArray<int32> PoleBonesIdx;
+	TArray<int32> IntersectBonesIdx;
+	
+	TArray<int32> AllChildrenWithinChain;
 };
 
 USTRUCT(BlueprintType, meta = (DisplayName = "Body Intersect Goals"))
@@ -83,6 +153,12 @@ struct FIKRetargetBodyIntersectIKOp : public FIKRetargetOpBase
 		const double InDeltaTime,
 		const TArray<FTransform>& InSourceGlobalPose,
 		TArray<FTransform>& OutTargetGlobalPose) override;
+	
+	UE_API virtual void RunAfterParent(
+		FIKRetargetProcessor& InProcessor,
+		const double InDeltaTime,
+		const TArray<FTransform>& InSourceGlobalPose,
+		TArray<FTransform>& OutTargetGlobalPose) override;
 
 	UE_API virtual void OnAddedToStack(const UIKRetargeter* InRetargetAsset, const FIKRetargetOpBase* InParentOp) override;
 
@@ -98,16 +174,13 @@ struct FIKRetargetBodyIntersectIKOp : public FIKRetargetOpBase
 	FIKRetargetBodyIntersectIKOpSettings Settings;
 
 private:
+	void InitPoleVectorIntersectors(const TArray<FResolvedBoneChain>& TargetChains, const FTargetSkeleton& InTargetSkeleton);
+	
 	void UpdateCacheSkelInfo(const FTargetSkeleton& InTargetSkeleton);
 	void UpdateTargetBoneMap(FName TargetBoneName);
 
 	FVector GoalLocBlendCompSpace(const FIKRigGoal* Goal, const FTransform& BoneTfm) const;
 	void SetGoalPosFromCompSpace(FIKRigGoal* Goal, const FTransform& BoneTfm, const FVector& CompSpaceLoc) const;
-	
-	static double CalcShapeSmallRadius(FKShapeElem* ShapeElem);
-	static double CalcShapeAvgRadius(FKShapeElem* ShapeElem);
-
-	static FKShapeElem* FindBodyShape(const UPhysicsAsset* PhysAsset, FName BoneName);
 
 #if WITH_EDITOR
 public:
@@ -136,6 +209,9 @@ private:
 	// Cache ik goals to check
 	TArray<const FIKRigGoal*> IntersectIKGoals;
 	
+	// Array (per-chain w/ pv bodies) of intersectors
+	TArray<FPoleVectorChainIntersector> PoleVectorChainIntersectors;
+	
 	// Cache source/target bone-names
 	TArray<FName> TargetBoneNames;
 
