BodyIntersectOp: Added checkbox for enabling/disabling body intersect (in order to only use pole vec).

#Synced-CL 49145592

diff --git a/Engine/Plugins/Experimental/Animation/RelativeIKOp/Source/BodyIntersectIKOp/Private/BodyIntersectIKOp.cpp b/Engine/Plugins/Experimental/Animation/RelativeIKOp/Source/BodyIntersectIKOp/Private/BodyIntersectIKOp.cpp
--- a/Engine/Plugins/Experimental/Animation/RelativeIKOp/Source/BodyIntersectIKOp/Private/BodyIntersectIKOp.cpp
+++ b/Engine/Plugins/Experimental/Animation/RelativeIKOp/Source/BodyIntersectIKOp/Private/BodyIntersectIKOp.cpp
@@ -157,6 +157,12 @@ void FIKRetargetBodyIntersectIKOp::Run(
 		
 		for (int j = 0; j < Settings.IntersectBodies.Num(); ++j)
 		{
+			if (!Settings.bEnableGoalIntersect)
+			{
+				// TODO: Move this test further out but still allow debug rendering of intersect bodies, etc
+				continue;
+			}
+			
 			FName BodyName = Settings.IntersectBodies[j];
 			if (!CacheTargetSkelIndices.Contains(BodyName))
 			{
diff --git a/Engine/Plugins/Experimental/Animation/RelativeIKOp/Source/BodyIntersectIKOp/Public/BodyIntersectIKOp.h b/Engine/Plugins/Experimental/Animation/RelativeIKOp/Source/BodyIntersectIKOp/Public/BodyIntersectIKOp.h
--- a/Engine/Plugins/Experimental/Animation/RelativeIKOp/Source/BodyIntersectIKOp/Public/BodyIntersectIKOp.h
+++ b/Engine/Plugins/Experimental/Animation/RelativeIKOp/Source/BodyIntersectIKOp/Public/BodyIntersectIKOp.h
@@ -77,6 +77,10 @@ struct FIKRetargetBodyIntersectIKOpSettings : public FIKRetargetOpSettingsBase
 	UPROPERTY(EditAnywhere, BlueprintReadWrite, Category = "Physics", meta=(ReinitializeOnEdit))
 	TObjectPtr<UPhysicsAsset> TargetPhysicsAssetOverride;
 
+	// Apply intersection delta to goals 
+	UPROPERTY(EditAnywhere, BlueprintReadWrite, Category = "Intersection", meta=(ReinitializeOnEdit))
+	bool bEnableGoalIntersect = true;
+	
 	// IK Goals Shape intersect against physics asset bodies 
 	UPROPERTY(EditAnywhere, BlueprintReadWrite, Category = "Intersection", meta=(ReinitializeOnEdit))
 	TArray<FIKGoalIntersectShapeSettings> IntersectGoalSettings;
