RelativeIK:
+ PreviewPropOps support for skeletonmesh prop animation

#Synced-CL 49175988

diff --git a/Engine/Plugins/Experimental/Animation/RelativeIKOp/Source/PreviewPropOp/Private/PreviewPropOp.cpp b/Engine/Plugins/Experimental/Animation/RelativeIKOp/Source/PreviewPropOp/Private/PreviewPropOp.cpp
--- a/Engine/Plugins/Experimental/Animation/RelativeIKOp/Source/PreviewPropOp/Private/PreviewPropOp.cpp
+++ b/Engine/Plugins/Experimental/Animation/RelativeIKOp/Source/PreviewPropOp/Private/PreviewPropOp.cpp
@@ -2,23 +2,26 @@
 
 #include "PreviewPropOp.h"
 
-#include "SceneView.h"
 #include "Retargeter/IKRetargetProcessor.h"
 
 #if WITH_EDITOR
+#include "AnimationRuntime.h"
+#include "BoneContainer.h"
+#include "BoneIndices.h"
+#include "BonePose.h"
 #include "PrimitiveDrawingUtils.h"
+#include "SceneView.h"
+#include "SkeletalMeshAttributes.h"
+#include "StaticMeshAttributes.h"
 #include "Animation/AnimInstance.h"
 #include "Animation/AnimMontage.h"
 #include "Animation/AnimSequence.h"
 #include "Animation/AnimSequenceBase.h"
 #include "Components/SkeletalMeshComponent.h"
 #include "Engine/EngineTypes.h"
+#include "Rendering/SkeletalMeshRenderData.h"
 #endif
 
-#include "StaticMeshAttributes.h"
-#include "AI/NavigationModifier.h"
-#include "GeometryCollection/Facades/CollectionUVFacade.h"
-
 #include UE_INLINE_GENERATED_CPP_BY_NAME(PreviewPropOp)
 
 #define LOCTEXT_NAMESPACE "PreviewPropOp"
@@ -83,9 +86,10 @@ bool FIKRetargetPreviewPropOp::Initialize(
 		return false;
 	}
 	
-	for (const FPreviewPropsData& PreviewData : Settings.PreviewProps)
+	for (int32 PropIdx = 0; PropIdx < Settings.PreviewProps.Num(); PropIdx++)
 	{
-		if (!PreviewData.PropMesh || PreviewData.SourceAttachBone.BoneName == NAME_None)
+		FPreviewPropsData& PreviewData = Settings.PreviewProps[PropIdx];
+		if ((!PreviewData.PropStaticMeshAsset && !PreviewData.PropSkeletalMeshAsset) || PreviewData.SourceAttachBone.BoneName == NAME_None)
 		{
 			continue;
 		}
@@ -100,6 +104,7 @@ bool FIKRetargetPreviewPropOp::Initialize(
 		int32 TargetAttachIdx = InTargetSkeleton.FindBoneIndexByName(TargetAttachBone);
 
 		FPropMeshData NewMesh;
+		NewMesh.PropIdx = PropIdx;
 		NewMesh.SourceAttachBoneIndex = SourceAttachIdx;
 		NewMesh.TargetAttachBoneIndex = TargetAttachIdx;
 		NewMesh.AttachTransform = PreviewData.AttachTransform;
@@ -134,9 +139,8 @@ void FIKRetargetPreviewPropOp::Run(
 	FScopeLock ScopeLock(&DebugDataMutex);
 	
 	DebugPropInfo.PropTfmInfo.Reset(PropMeshData.Num());
-
-	int32 PropIdx = 0;
-	for (const FPropMeshData& PropMesh : PropMeshData)
+	
+	for (FPropMeshData& PropMesh : PropMeshData)
 	{
 		FPropTfmInfo& PropInfo = DebugPropInfo.PropTfmInfo.AddDefaulted_GetRef();
 		
@@ -151,10 +155,22 @@ void FIKRetargetPreviewPropOp::Run(
 		}
 		
 		// TODO: Move this into processing for props w/ anim seq only 
-		float TestPropAnimLength = 0.5;
-		float TestPropPlayhead = ComputePropAnimPlayhead(TestPropAnimLength, true);
-
-		++PropIdx;
+		UAnimSequence* PropAnimation = Settings.PreviewProps[PropMesh.PropIdx].PropAnimSequence;
+		USkeletalMesh* PropSkeletalMeshAsset = Settings.PreviewProps[PropMesh.PropIdx].PropSkeletalMeshAsset;
+		if (PropAnimation && PropSkeletalMeshAsset)
+		{
+			float TestPropAnimLength = PropAnimation->GetPlayLength();
+			float TestPropPlayhead = ComputePropAnimPlayhead(TestPropAnimLength, Settings.PreviewProps[PropMesh.PropIdx].PropAnimStartTime, Settings.PreviewProps[PropMesh.PropIdx].PropAnimIsLooping);
+			if (TestPropPlayhead<0.f)
+			{
+				continue;
+			}
+			
+			TArray<FMatrix44f> CacheToLocals;
+			TArray<FVector3f> PropVLocations;
+			GetRefToAnimPoseMatrices(PropAnimation, TestPropPlayhead, PropSkeletalMeshAsset, CacheToLocals);
+			UpdateSkinnedVertices(PropSkeletalMeshAsset, PropMesh, CacheToLocals);
+		}
 	}
 	
 #endif // WITH_EDITOR
@@ -267,7 +283,7 @@ void FIKRetargetPreviewPropOp::UpdateAnimSeqPlayhead(UAnimInstance* SourceAnimIn
 	}
 }
 
-float FIKRetargetPreviewPropOp::ComputePropAnimPlayhead(float PropAnimLength, bool bLooping) const
+float FIKRetargetPreviewPropOp::ComputePropAnimPlayhead(float PropAnimLength, float PropAnimStartTime, bool bLooping) const
 {
 	if (!CacheSourceCharacterAnim)
 	{
@@ -276,17 +292,67 @@ float FIKRetargetPreviewPropOp::ComputePropAnimPlayhead(float PropAnimLength, bo
 	
 	if (!bLooping)
 	{
-		return FMath::Min(PropAnimLength, SourceAnimPlayhead);
+		return FMath::Clamp(SourceAnimPlayhead-PropAnimStartTime, 0.0f, PropAnimLength);
 	}
 	
 	// NOTE: This will pop at the end of char sequence if SourceAnimLength is not an integer multiple of PropAnimLength
 	// TODO: Potentially compute a playrate to force correct looping behaviour? 
-	return FMath::Fmod(SourceAnimPlayhead, PropAnimLength);
+	return FMath::Max(FMath::Fmod(SourceAnimPlayhead-PropAnimStartTime, PropAnimLength), 0.0f);
+}
+
+void FIKRetargetPreviewPropOp::GetAnimSeqFramePose(const UAnimSequenceBase* AnimSeq, double Time, TArray<FName>& OutBones, TArray<FTransform>& OutPose) const
+{
+	const FReferenceSkeleton& Skel = AnimSeq->GetSkeleton()->GetReferenceSkeleton();
+	int32 NumBones = Skel.GetNum();
+	
+	OutBones.Reset(NumBones);
+	OutPose.Reset(NumBones);
+
+	// Init bone container for pulling animseq data
+	TArray<FBoneIndexType> BoneIndices;
+	BoneIndices.Reserve(NumBones);
+	for (FBoneIndexType BoneIndex = 0; BoneIndex < NumBones; ++BoneIndex)
+	{
+		BoneIndices.Add(BoneIndex);
+		OutBones.Add(Skel.GetBoneName(BoneIndex));
+	}
+	
+	FBoneContainer BoneContainer;
+	BoneContainer.InitializeTo(BoneIndices, UE::Anim::FCurveFilterSettings(UE::Anim::ECurveFilterMode::DisallowAll), *AnimSeq->GetSkeleton());
+	BoneContainer.SetUseRAWData(false);
+	
+	// Setup bone container
+	FCompactPose CompactPose;
+	FBlendedCurve Curves;
+	UE::Anim::FStackAttributeContainer Attributes;
+	FAnimationPoseData FramePoseData(CompactPose, Curves, Attributes);
+	
+	CompactPose.SetBoneContainer(&BoneContainer);
+	Curves.InitFrom(BoneContainer);
+	
+	FAnimExtractContext FrameExtractionCtx(Time, false);
+	AnimSeq->GetAnimationPose(FramePoseData, FrameExtractionCtx);
+	
+	FramePoseData.GetPose().CopyBonesTo(OutPose);
+	for (int32 BoneIndex = 0; BoneIndex < NumBones; ++BoneIndex)
+	{
+		const int32 ParentBoneIndex = Skel.GetParentIndex(BoneIndex);
+
+		FTransform ParentTransform = FTransform::Identity;
+		if (ParentBoneIndex != INDEX_NONE)
+		{
+			// TODO: Double-check parents always befor children in list
+			FString ParentBoneName = Skel.GetBoneName(ParentBoneIndex).ToString();
+			ParentTransform = OutPose[ParentBoneIndex];
+		}
+		
+		OutPose[BoneIndex] = OutPose[BoneIndex] * ParentTransform;
+	}
 }
 
 bool FIKRetargetPreviewPropOp::GetMeshVertTris(FPropMeshData& OutMeshData, const FPreviewPropsData& PreviewProp)
 {
-	const FMeshDescription* Desc = PreviewProp.PropMesh->GetMeshDescription(0);
+	const FMeshDescription* Desc = PreviewProp.GetMeshDescription(0);
 	if (!Desc)
 	{
 		return false;
@@ -296,6 +362,9 @@ bool FIKRetargetPreviewPropOp::GetMeshVertTris(FPropMeshData& OutMeshData, const
 	int32 NumTris = Desc->Triangles().Num();
 	
 	OutMeshData.Vertices.Reset(NumVerts);
+	OutMeshData.Vertices.AddUninitialized(NumVerts);
+	OutMeshData.RestVLocation.Reset(NumVerts);
+	OutMeshData.RestVLocation.AddUninitialized(NumVerts);
 	OutMeshData.Indices.SetNumUninitialized(3*NumTris);
 
 	// TODO: Almost certainly incorrectly "welding" all vert instances into whatever the first seen vert is!
@@ -316,8 +385,6 @@ bool FIKRetargetPreviewPropOp::GetMeshVertTris(FPropMeshData& OutMeshData, const
 			const FVertexID VertID = TriVerts[Idx];
 			if (!TriVertMap.Contains(VertID))
 			{
-				int32 NewVertIdx = OutMeshData.Vertices.Num();
-
 				FDynamicMeshVertex NewVert{Desc->GetVertexPosition(VertID)};
 				if (TexCoords.IsValid())
 				{
@@ -338,10 +405,10 @@ bool FIKRetargetPreviewPropOp::GetMeshVertTris(FPropMeshData& OutMeshData, const
 					NewVert.TangentZ.Vector.W = static_cast<int8>(Vsgn);
 				}
 
-				OutMeshData.Vertices.Add(NewVert);
-				TriVertMap.Add(VertID, NewVertIdx);
+				OutMeshData.Vertices[VertID] = NewVert;
+				OutMeshData.RestVLocation[VertID] = Desc->GetVertexPosition(VertID);
 			}
-			OutMeshData.Indices[TriIdx*3 + Idx] = TriVertMap[VertID];
+			OutMeshData.Indices[TriIdx*3 + Idx] = VertID;
 		}
 		
 		++TriIdx;
@@ -429,6 +496,73 @@ void FIKRetargetPreviewPropOp::DebugDraw(
 	}
 }
 
+void FIKRetargetPreviewPropOp::UpdateSkinnedVertices(USkeletalMesh* SkeletalMeshAsset, FPropMeshData& PropMeshDataIn, const TArray<FMatrix44f>& RefToPoseMatrices)
+{
+	FMeshDescription* MeshDescriptionPtr = SkeletalMeshAsset->GetMeshDescription(0);
+	if (!MeshDescriptionPtr)
+	{
+		return;
+	}
+	const int32 NumVertices = MeshDescriptionPtr->Vertices().Num();
+
+	//Extract skin weights
+	FReferenceSkeleton& RefSkeleton = SkeletalMeshAsset->GetRefSkeleton();
+	FSkeletalMeshAttributes MeshAttribs(*MeshDescriptionPtr);
+	FSkinWeightsVertexAttributesRef VertexSkinWeights = MeshAttribs.GetVertexSkinWeights();
+	
+	//CPU skinned vertices
+	const int32 NumBones = MeshAttribs.GetNumBones();
+	for (int32 VertexIndex = 0; VertexIndex < NumVertices; ++VertexIndex)
+	{
+		FVector3f NewPosition = FVector3f::ZeroVector;
+		FVertexBoneWeights BoneWeights = VertexSkinWeights.Get(FVertexID(VertexIndex));
+		const int32 InfluenceCount = BoneWeights.Num();
+		for (int32 InfluenceIndex = 0; InfluenceIndex < InfluenceCount; ++InfluenceIndex)
+		{
+			int32 InfluenceBoneIndex = BoneWeights[InfluenceIndex].GetBoneIndex();
+			float InfluenceBoneWeight = BoneWeights[InfluenceIndex].GetWeight();
+
+			const FMatrix44f& RefToPose = RefToPoseMatrices[InfluenceBoneIndex];
+			NewPosition += RefToPose.TransformPosition(PropMeshDataIn.RestVLocation[VertexIndex]) * InfluenceBoneWeight;
+		}
+		PropMeshDataIn.Vertices[VertexIndex].Position = NewPosition;
+	}
+}
+
+void FIKRetargetPreviewPropOp::GetRefToAnimPoseMatrices(UAnimSequence* InAnimation, float AnimPlayhead, USkeletalMesh* SkeletalMeshAsset, TArray<FMatrix44f>& OutRefToPose) const
+{
+	const FReferenceSkeleton& RefSkeleton = SkeletalMeshAsset->GetRefSkeleton();
+	const TArray<FMatrix44f>& RefBasesInvMatrix = SkeletalMeshAsset->GetRefBasesInvMatrix();
+	
+	OutRefToPose.Init(FMatrix44f::Identity, RefBasesInvMatrix.Num());
+	
+	TArray<FName> AnimBoneNames;
+	TArray<FTransform> AnimPose;
+	GetAnimSeqFramePose(InAnimation, AnimPlayhead, AnimBoneNames, AnimPose);
+	
+	const FSkeletalMeshLODRenderData& LOD = SkeletalMeshAsset->GetResourceForRendering()->LODRenderData[0];
+	const TArray<FBoneIndexType>& RequiredBoneIndices = LOD.ActiveBoneIndices;
+	for (int32 BoneIndex = 0; BoneIndex < RequiredBoneIndices.Num(); BoneIndex++)
+	{
+		const int32 ThisBoneIndex = RequiredBoneIndices[BoneIndex];
+		if (RefBasesInvMatrix.IsValidIndex(ThisBoneIndex))
+		{
+			FName ThisBoneName = RefSkeleton.GetBoneName(ThisBoneIndex);
+			int32 PoseIdx = AnimBoneNames.Find(ThisBoneName);
+			if (PoseIdx == INDEX_NONE)
+			{
+				continue;
+			}
+			
+			OutRefToPose[ThisBoneIndex] = static_cast<FMatrix44f>(AnimPose[PoseIdx].ToMatrixWithScale());
+		}
+	}
+
+	for (int32 ThisBoneIndex = 0; ThisBoneIndex < OutRefToPose.Num(); ++ThisBoneIndex)
+	{
+		OutRefToPose[ThisBoneIndex] = RefBasesInvMatrix[ThisBoneIndex] * OutRefToPose[ThisBoneIndex];
+	}
+}
 #endif //WITH_EDITOR
 
 FIKRetargetPreviewPropOpSettings UIKRetargetPreviewPropController::GetSettings()
diff --git a/Engine/Plugins/Experimental/Animation/RelativeIKOp/Source/PreviewPropOp/Public/PreviewPropOp.h b/Engine/Plugins/Experimental/Animation/RelativeIKOp/Source/PreviewPropOp/Public/PreviewPropOp.h
--- a/Engine/Plugins/Experimental/Animation/RelativeIKOp/Source/PreviewPropOp/Public/PreviewPropOp.h
+++ b/Engine/Plugins/Experimental/Animation/RelativeIKOp/Source/PreviewPropOp/Public/PreviewPropOp.h
@@ -4,6 +4,7 @@
 #include "DynamicMeshBuilder.h"
 #include "Animation/BoneReference.h"
 #include "Engine/StaticMesh.h"
+#include "Engine/SkeletalMesh.h"
 #include "Materials/MaterialInstanceDynamic.h"
 #include "Retargeter/IKRetargetOps.h"
 
@@ -16,6 +17,7 @@
 struct FIKRigGoal;
 struct FKShapeElem;
 class UAnimInstance;
+class UAnimSequence;
 class UAnimSequenceBase;
 class UPhysicsAsset;
 class USkeletalMeshComponent;
@@ -35,8 +37,10 @@ struct FDebugPropDrawInfo
 
 struct FPropMeshData
 {
+	int32 PropIdx;
 	TArray<uint32> Indices;
 	TArray<FDynamicMeshVertex> Vertices;
+	TArray<FVector3f> RestVLocation;
 
 	int32 SourceAttachBoneIndex;
 	int32 TargetAttachBoneIndex;
@@ -51,7 +55,11 @@ struct FPreviewPropsData
 // #if WITH_EDITORONLY_DATA
 	// Prop mesh to preview
 	UPROPERTY(EditAnywhere, BlueprintReadWrite, Category = "Prop", meta=(ReinitializeOnEdit))
-	TObjectPtr<UStaticMesh> PropMesh;
+	TObjectPtr<UStaticMesh> PropStaticMeshAsset;
+	
+	// Prop mesh to preview
+	UPROPERTY(EditAnywhere, BlueprintReadWrite, Category = "Prop", meta=(ReinitializeOnEdit))
+	TObjectPtr<USkeletalMesh> PropSkeletalMeshAsset;
 
 	// Source skeletal mesh bone to attach prop to
 	UPROPERTY(EditAnywhere, Category = "Prop", meta=(ReinitializeOnEdit))
@@ -60,7 +68,36 @@ struct FPreviewPropsData
 	// Local transform of prop relative to attach bone
 	UPROPERTY(EditAnywhere, BlueprintReadWrite, Category = "Prop", meta=(ReinitializeOnEdit))
 	FTransform AttachTransform;
+	
+	// Prop mesh to anim sequence
+	UPROPERTY(EditAnywhere, BlueprintReadWrite, Category = "Prop", meta=(ReinitializeOnEdit))
+	TObjectPtr<UAnimSequence> PropAnimSequence;
+	
+	// Prop mesh to anim sequence start time
+	UPROPERTY(EditAnywhere, BlueprintReadWrite, Category = "Prop", meta=(ReinitializeOnEdit))
+	float PropAnimStartTime = 0.0f;
+	
+	// Prop mesh to anim sequence is looping
+	UPROPERTY(EditAnywhere, BlueprintReadWrite, Category = "Prop", meta=(ReinitializeOnEdit))
+	bool PropAnimIsLooping = true;
 // ##endif //WITH_EDITORONLY_DATA
+
+#if WITH_EDITOR
+	const FMeshDescription* GetMeshDescription(int32 LodIndex) const
+	{
+		if (PropSkeletalMeshAsset != nullptr)
+		{
+			return PropSkeletalMeshAsset->GetMeshDescription(LodIndex);
+		}
+	
+		if (PropStaticMeshAsset != nullptr)
+		{
+			return PropStaticMeshAsset->GetMeshDescription(LodIndex);
+		}
+
+		return nullptr;
+	}
+#endif //WITH_EDITOR
 };
 
 USTRUCT(BlueprintType, meta = (DisplayName = "Preview Prop Settings"))
@@ -129,7 +166,9 @@ private:
 
 #if WITH_EDITOR
 	void UpdateAnimSeqPlayhead(UAnimInstance* SourceAnimInstance);
-	float ComputePropAnimPlayhead(float PropAnimLength, bool bLooping) const;
+	float ComputePropAnimPlayhead(float PropAnimLength, float PropAnimStartTime, bool bLooping) const;
+	
+	void GetAnimSeqFramePose(const UAnimSequenceBase* AnimSeq, double Time, TArray<FName>& OutBones, TArray<FTransform>& OutPose) const;
 	
 	static bool GetMeshVertTris(FPropMeshData& OutMeshData, const FPreviewPropsData& PreviewProp);
 	
@@ -142,6 +181,8 @@ public:
 		const FIKRetargetDebugDrawState& InEditorState) const override;
 
 	virtual bool HasDebugDrawing() override { return true; };
+	void UpdateSkinnedVertices(USkeletalMesh* SkeletalMeshAsset, FPropMeshData& PropMeshDataIn, const TArray<FMatrix44f>& RefToPoseMatrices);
+	void GetRefToAnimPoseMatrices(UAnimSequence* InAnimation, float AnimPlayhead, USkeletalMesh* SkeletalMeshAsset, TArray<FMatrix44f>& OutRefToPose) const;
 
 private:
 	// Bone transforms for prop preview
