Relative IK:
+ Support PropModifierBaking with Prop Anim

#Synced-CL 49283367

diff --git a/Engine/Plugins/Experimental/Animation/RelativeIKOp/Source/RelativeBodyAnimUtils/Private/RelativePropsAnimModifier.cpp b/Engine/Plugins/Experimental/Animation/RelativeIKOp/Source/RelativeBodyAnimUtils/Private/RelativePropsAnimModifier.cpp
--- a/Engine/Plugins/Experimental/Animation/RelativeIKOp/Source/RelativeBodyAnimUtils/Private/RelativePropsAnimModifier.cpp
+++ b/Engine/Plugins/Experimental/Animation/RelativeIKOp/Source/RelativeBodyAnimUtils/Private/RelativePropsAnimModifier.cpp
@@ -76,6 +76,19 @@ struct FPropsMeshAndTree
 		Converter.Convert(MeshDesc, Mesh);
 		AABBTree.Build();
 	}
+	
+	// Update all vertex positions with a lambda
+	void UpdateVertexPositions(TArray<FVector3f>& VLocation)
+	{
+		for (int32 Vid : Mesh.VertexIndicesItr())
+		{
+			// FVector3d OldPos = Mesh.GetVertex(Vid);
+			Mesh.SetVertex(Vid, FVector3d(VLocation[Vid]));
+		}
+
+		// After modifying vertices, rebuild the tree
+		AABBTree.Build();
+	}
 };
 
 const FPositionVertexBuffer& FPropsInfo::GetPositionVertexBuffer()
@@ -113,6 +126,101 @@ FMeshDescription* FPropsInfo::GetMeshDescription(int32 LodIndex)
 	return nullptr;
 }
 
+float FPropsInfo::ComputePropAnimPlayhead(float CharacterPlayhead) const
+{
+	if (!PropAnimSequence)
+	{
+		return -1.f;
+	}
+	
+	if (!PropAnimIsLooping)
+	{
+		return FMath::Min(PropAnimSequence->GetPlayLength(), CharacterPlayhead-PropAnimStartTime);
+	}
+	
+	return FMath::Fmod(CharacterPlayhead-PropAnimStartTime, PropAnimSequence->GetPlayLength());
+}
+
+void FPropsInfo::GetSkinnedVertices(const TArray<FMatrix44f>& RefToPoseMatrices, TArray<FVector3f>& VLocation)
+{
+	if (!PropSkeletalMeshAsset)
+	{
+		return;
+	}
+	FMeshDescription* MeshDescriptionPtr = PropSkeletalMeshAsset->GetMeshDescription(0);
+	if (!MeshDescriptionPtr)
+	{
+		return;
+	}
+	const int32 NumVertices = MeshDescriptionPtr->Vertices().Num();
+
+	//Extract skin weights
+	FReferenceSkeleton& RefSkeleton = PropSkeletalMeshAsset->GetRefSkeleton();
+	FSkeletalMeshAttributes MeshAttribs(*MeshDescriptionPtr);
+	FSkinWeightsVertexAttributesRef VertexSkinWeights = MeshAttribs.GetVertexSkinWeights();
+	
+	//CPU skinned vertices
+	const int32 NumBones = MeshAttribs.GetNumBones();
+	for (int32 VertexIndex = 0; VertexIndex < NumVertices; ++VertexIndex)
+	{
+		FVector3f NewPosition = FVector3f::ZeroVector;
+		FVertexBoneWeights BoneWeights = VertexSkinWeights.Get(FVertexID(VertexIndex));
+		const int32 InfluenceCount = BoneWeights.Num();
+		for (int32 InfluenceIndex = 0; InfluenceIndex < InfluenceCount; ++InfluenceIndex)
+		{
+			int32 InfluenceBoneIndex = BoneWeights[InfluenceIndex].GetBoneIndex();
+			float InfluenceBoneWeight = BoneWeights[InfluenceIndex].GetWeight();
+
+			const FMatrix44f& RefToPose = RefToPoseMatrices[InfluenceBoneIndex];
+			NewPosition += RefToPose.TransformPosition(MeshDescriptionPtr->GetVertexPosition(VertexIndex)) * InfluenceBoneWeight;
+		}
+		VLocation[VertexIndex] = NewPosition;
+	}
+}
+
+void FPropsInfo::GetRefToAnimPoseMatrices(float AnimPlayhead, TArray<FMatrix44f>& OutRefToPose) const
+{
+	if (!PropAnimSequence || !PropSkeletalMeshAsset)
+	{
+		return;
+	}
+	
+	static FAnimPoseEvaluationOptions AnimPoseEvalOptions{
+		EAnimDataEvalType::Raw,
+		true,
+		false,
+		false,
+		nullptr,
+		true,
+		false
+	};
+	
+	FAnimPose AnimPose;
+	UAnimPoseExtensions::GetAnimPoseAtTime(PropAnimSequence, AnimPlayhead, AnimPoseEvalOptions, AnimPose);
+	const FReferenceSkeleton& RefSkeleton = PropSkeletalMeshAsset->GetRefSkeleton();
+	const TArray<FMatrix44f>& RefBasesInvMatrix = PropSkeletalMeshAsset->GetRefBasesInvMatrix();
+	
+	OutRefToPose.Init(FMatrix44f::Identity, RefBasesInvMatrix.Num());
+	
+	const FSkeletalMeshLODRenderData& LOD = PropSkeletalMeshAsset->GetResourceForRendering()->LODRenderData[0];
+	const TArray<FBoneIndexType>& RequiredBoneIndices = LOD.ActiveBoneIndices;
+	for (int32 BoneIndex = 0; BoneIndex < RequiredBoneIndices.Num(); BoneIndex++)
+	{
+		const int32 ThisBoneIndex = RequiredBoneIndices[BoneIndex];
+		if (RefBasesInvMatrix.IsValidIndex(ThisBoneIndex))
+		{
+			FName ThisBoneName = RefSkeleton.GetBoneName(ThisBoneIndex);
+			const FTransform ThisBoneTransform = UAnimPoseExtensions::GetBonePose(AnimPose, ThisBoneName, EAnimPoseSpaces::World);
+			OutRefToPose[ThisBoneIndex] = static_cast<FMatrix44f>(ThisBoneTransform.ToMatrixWithScale());
+		}
+	}
+
+	for (int32 ThisBoneIndex = 0; ThisBoneIndex < OutRefToPose.Num(); ++ThisBoneIndex)
+	{
+		OutRefToPose[ThisBoneIndex] = RefBasesInvMatrix[ThisBoneIndex] * OutRefToPose[ThisBoneIndex];
+	}
+}
+
 void URelativePropsAnimModifier::OnApply_Implementation(UAnimSequence* InAnimation)
 {
 	if (NotifyClass)
@@ -238,7 +346,7 @@ void URelativePropsAnimModifier::OnApply_Implementation(UAnimSequence* InAnimati
 				{
 					PropsTreeData[PropsIdx].Init(PropsMeshDescription);
 				}
-
+	
 				for (int32 BodyIndex1 = 0; BodyIndex1 < NumBodies; ++BodyIndex1)
 				{
 					if (BodyIndicesToIgnore[BodyIndex1])
@@ -305,6 +413,22 @@ void URelativePropsAnimModifier::OnApply_Implementation(UAnimSequence* InAnimati
 					FTransform AttachBoneTransform = UAnimPoseExtensions::GetBonePose(AnimPose, PropsData[PropsIdx].AttachBone.BoneName, EAnimPoseSpaces::World);
 					FTransform PropsLocalToWorld = PropsData[PropsIdx].AttachTransform*AttachBoneTransform;
 					const FPositionVertexBuffer& PositionVertexBuffer = PropsData[PropsIdx].GetPositionVertexBuffer();
+					
+					if (PropsData[PropsIdx].PropSkeletalMeshAsset && PropsData[PropsIdx].PropAnimSequence)
+					{
+						float PropPlayhead = PropsData[PropsIdx].ComputePropAnimPlayhead(SampleTime);
+						TArray<FMatrix44f> PropsCacheToLocals;
+						PropsData[PropsIdx].GetRefToAnimPoseMatrices(PropPlayhead, PropsCacheToLocals);
+						PropsData[PropsIdx].GetSkinnedVertices(PropsCacheToLocals, PropsVLocations[PropsIdx]);
+						
+						PropsTreeData[PropsIdx].UpdateVertexPositions(PropsVLocations[PropsIdx]);
+						for (int32 i = 0; i < PropsVLocations[PropsIdx].Num(); ++i)
+						{
+							FVector LocalPos = PropsData[PropsIdx].AttachTransform.TransformPosition(FVector(PropsVLocations[PropsIdx][i]));
+							PropsVLocations[PropsIdx][i] = FVector3f(AttachBoneTransform.TransformPosition(LocalPos));
+						}
+					}
+					
 					for (int32 BodyIndex = 0; BodyIndex < NumBodies; ++BodyIndex)
 					{
 						if (BodyIndicesToIgnore[BodyIndex])
@@ -360,8 +484,8 @@ void URelativePropsAnimModifier::OnApply_Implementation(UAnimSequence* InAnimati
 						int32 LowestVertexIndex = FindLowestIndex(PropsVLocations[PropsIdx]);
 						int32 HighestVertexIndex=FindHighestIndex(PropsVLocations[PropsIdx]);
 
-						BakedAnimNotifyInfo->PropsPairsLocalReference[Offset + 0] = PropOffsetTfm.TransformPosition(PositionVertexBuffer.VertexPosition(LowestVertexIndex));
-						BakedAnimNotifyInfo->PropsPairsLocalReference[Offset + 1] = PropOffsetTfm.TransformPosition(PositionVertexBuffer.VertexPosition(HighestVertexIndex));
+						BakedAnimNotifyInfo->PropsPairsLocalReference[Offset + 0] = FVector3f(AttachBoneTransform.InverseTransformPosition(FVector(PropsVLocations[PropsIdx][LowestVertexIndex])));
+						BakedAnimNotifyInfo->PropsPairsLocalReference[Offset + 1] = FVector3f(AttachBoneTransform.InverseTransformPosition(FVector(PropsVLocations[PropsIdx][HighestVertexIndex])));
 
 						Count++;
 					}
diff --git a/Engine/Plugins/Experimental/Animation/RelativeIKOp/Source/RelativeBodyAnimUtils/Public/RelativePropsAnimModifier.h b/Engine/Plugins/Experimental/Animation/RelativeIKOp/Source/RelativeBodyAnimUtils/Public/RelativePropsAnimModifier.h
--- a/Engine/Plugins/Experimental/Animation/RelativeIKOp/Source/RelativeBodyAnimUtils/Public/RelativePropsAnimModifier.h
+++ b/Engine/Plugins/Experimental/Animation/RelativeIKOp/Source/RelativeBodyAnimUtils/Public/RelativePropsAnimModifier.h
@@ -36,10 +36,26 @@ struct FPropsInfo
 	// Local transform of prop relative to attach bone
 	UPROPERTY(EditAnywhere, BlueprintReadWrite, Category = "Prop", meta=(ReinitializeOnEdit))
 	FTransform AttachTransform;
+	
+	// Prop mesh to anim sequence
+	UPROPERTY(EditAnywhere, BlueprintReadWrite, Category = "Prop", meta=(ReinitializeOnEdit))
+	TObjectPtr<UAnimSequence> PropAnimSequence;
+	
+	// Prop mesh to anim sequence start time
+	UPROPERTY(EditAnywhere, BlueprintReadWrite, Category = "Prop", meta=(ReinitializeOnEdit))
+	float PropAnimStartTime = 0.0f;
+	
+	// Prop mesh to anim sequence is looping
+	UPROPERTY(EditAnywhere, BlueprintReadWrite, Category = "Prop", meta=(ReinitializeOnEdit))
+	bool PropAnimIsLooping = true;
 
 public:
 	const FPositionVertexBuffer& GetPositionVertexBuffer();
 	FMeshDescription* GetMeshDescription(int32 LodIndex);
+	
+	float ComputePropAnimPlayhead(float CharacterPlayhead) const;
+	void GetSkinnedVertices(const TArray<FMatrix44f>& RefToPoseMatrices, TArray<FVector3f>& VLocation);
+	void GetRefToAnimPoseMatrices(float AnimPlayhead, TArray<FMatrix44f>& OutRefToPose) const;
 };
 
 /**
