RelativeIK: Interpolation Smooth

#Synced-CL 49697779

diff --git a/Engine/Plugins/Experimental/Animation/RelativeIKOp/Source/RelativeIKOp/Private/RelativeIKOp.cpp b/Engine/Plugins/Experimental/Animation/RelativeIKOp/Source/RelativeIKOp/Private/RelativeIKOp.cpp
--- a/Engine/Plugins/Experimental/Animation/RelativeIKOp/Source/RelativeIKOp/Private/RelativeIKOp.cpp
+++ b/Engine/Plugins/Experimental/Animation/RelativeIKOp/Source/RelativeIKOp/Private/RelativeIKOp.cpp
@@ -1316,7 +1316,19 @@ TConstArrayView<FVector3f> FIKRetargetRelativeIKOp::ApplyTemporalSmoothing(float
 	
 	if (Settings.TemporalSmoothingRadius <= 0)
 	{
-		return {PairLocalVerts.GetData() + NumBodies * SampleIdx, NumBodies};
+		int32 SampleIdx2 = FMath::Clamp(SampleIdx+1, 0, NumSamples-1);
+		if (SampleIdx == SampleIdx2)
+		{
+			return {PairLocalVerts.GetData() + NumBodies * SampleIdx, NumBodies};
+		}
+		
+		float lambda = SampleIdx2-Time * SampleRate;
+		SmoothedPoints.SetNumUninitialized(NumBodies);
+		for (int BodyIdx = 0; BodyIdx < NumBodies; ++BodyIdx)
+		{
+			SmoothedPoints[BodyIdx] = lambda * PairLocalVerts[NumBodies * SampleIdx + BodyIdx] + (1.f-lambda) * PairLocalVerts[NumBodies * SampleIdx2 + BodyIdx];
+		}
+		return SmoothedPoints;
 	}
 	
 	SmoothedPoints.SetNumUninitialized(NumBodies);
@@ -1328,7 +1340,7 @@ TConstArrayView<FVector3f> FIKRetargetRelativeIKOp::ApplyTemporalSmoothing(float
 		{
 			// Finite support Gaussian smoothing (could approximate w/ triangle or bake out smoothed data)
 			const float Sigma = 0.25f * static_cast<float>(Settings.TemporalSmoothingRadius);
-			const float SmoothDist = static_cast<float>(FrameIdx - SampleIdx);
+			const float SmoothDist = static_cast<float>(FrameIdx - Time * SampleRate);
 			const float SmoothWeight = FMath::Exp(-0.5f*(SmoothDist*SmoothDist) / (Sigma*Sigma));
 			SmoothedPoints[BodyIdx] += SmoothWeight * PairLocalVerts[NumBodies * FrameIdx + BodyIdx];
 			SumWeight += SmoothWeight;
