quick fix for RIK

#Synced-CL 49735487

diff --git a/Engine/Plugins/Experimental/Animation/RelativeIKOp/Source/RelativeIKOp/Private/RelativeIKOp.cpp b/Engine/Plugins/Experimental/Animation/RelativeIKOp/Source/RelativeIKOp/Private/RelativeIKOp.cpp
--- a/Engine/Plugins/Experimental/Animation/RelativeIKOp/Source/RelativeIKOp/Private/RelativeIKOp.cpp
+++ b/Engine/Plugins/Experimental/Animation/RelativeIKOp/Source/RelativeIKOp/Private/RelativeIKOp.cpp
@@ -1322,7 +1322,7 @@ TConstArrayView<FVector3f> FIKRetargetRelativeIKOp::ApplyTemporalSmoothing(float
 			return {PairLocalVerts.GetData() + NumBodies * SampleIdx, NumBodies};
 		}
 		
-		float lambda = SampleIdx2-Time * SampleRate;
+		float lambda = static_cast<float>(SampleIdx2) - Time * SampleRate;
 		SmoothedPoints.SetNumUninitialized(NumBodies);
 		for (int BodyIdx = 0; BodyIdx < NumBodies; ++BodyIdx)
 		{
@@ -1340,7 +1340,7 @@ TConstArrayView<FVector3f> FIKRetargetRelativeIKOp::ApplyTemporalSmoothing(float
 		{
 			// Finite support Gaussian smoothing (could approximate w/ triangle or bake out smoothed data)
 			const float Sigma = 0.25f * static_cast<float>(Settings.TemporalSmoothingRadius);
-			const float SmoothDist = static_cast<float>(FrameIdx - Time * SampleRate);
+			const float SmoothDist = static_cast<float>(FrameIdx) - Time * SampleRate;
 			const float SmoothWeight = FMath::Exp(-0.5f*(SmoothDist*SmoothDist) / (Sigma*Sigma));
 			SmoothedPoints[BodyIdx] += SmoothWeight * PairLocalVerts[NumBodies * FrameIdx + BodyIdx];
 			SumWeight += SmoothWeight;
