RelativePelvicMotionOp: Fix for incorrect index check and settings update.

#Synced-CL 50148284

diff --git a/Engine/Plugins/Experimental/Animation/RelativeIKOp/Source/RelativeIKOp/Private/RelativePelvicMotionOp.cpp b/Engine/Plugins/Experimental/Animation/RelativeIKOp/Source/RelativeIKOp/Private/RelativePelvicMotionOp.cpp
--- a/Engine/Plugins/Experimental/Animation/RelativeIKOp/Source/RelativeIKOp/Private/RelativePelvicMotionOp.cpp
+++ b/Engine/Plugins/Experimental/Animation/RelativeIKOp/Source/RelativeIKOp/Private/RelativePelvicMotionOp.cpp
@@ -589,7 +589,7 @@ void FIKRetargetPelvicMotionOp::ApplyAverageBodyPairTargets(
 	TArray<FDebugRelativeTargetPairSpace>& DebugPairSpaces)
 {
 	int32 FindEffectIdx = CacheSourceEffectBones.Find(SourcePelvicBone);
-	if (CacheSourceEffectBones.IsValidIndex(FindEffectIdx))
+	if (!CacheSourceBodyEffectVertIdx.IsValidIndex(FindEffectIdx))
 	{
 		return;
 	}
diff --git a/Engine/Plugins/Experimental/Animation/RelativeIKOp/Source/RelativeIKOp/Public/RelativePelvicMotionOp.h b/Engine/Plugins/Experimental/Animation/RelativeIKOp/Source/RelativeIKOp/Public/RelativePelvicMotionOp.h
--- a/Engine/Plugins/Experimental/Animation/RelativeIKOp/Source/RelativeIKOp/Public/RelativePelvicMotionOp.h
+++ b/Engine/Plugins/Experimental/Animation/RelativeIKOp/Source/RelativeIKOp/Public/RelativePelvicMotionOp.h
@@ -75,7 +75,7 @@ struct FIKRetargetPelvicMotionOpSettings : public FIKRetargetOpSettingsBase
 	bool bIgnoreSourceScale = true;
 
 	// Alpha from contact-body to secondary body representation in contribution pairs
-	UPROPERTY(EditAnywhere, BlueprintReadWrite, Category = "Parameters")
+	UPROPERTY(EditAnywhere, BlueprintReadWrite, Category = "Parameters", meta=(ClampMin = "0", ClampMax = "1"))
 	double RetargetContactAlpha = 0.5;
 
 	// Alpha between primary and secondary pair distance relationship
