RelativePelvicMotionOp: Fix for loctext conflict.

#Synced-CL 50148284

diff --git a/Engine/Plugins/Experimental/Animation/RelativeIKOp/Source/RelativeIKOp/Private/RelativePelvicMotionOp.cpp b/Engine/Plugins/Experimental/Animation/RelativeIKOp/Source/RelativeIKOp/Private/RelativePelvicMotionOp.cpp
--- a/Engine/Plugins/Experimental/Animation/RelativeIKOp/Source/RelativeIKOp/Private/RelativePelvicMotionOp.cpp
+++ b/Engine/Plugins/Experimental/Animation/RelativeIKOp/Source/RelativeIKOp/Private/RelativePelvicMotionOp.cpp
@@ -90,7 +90,7 @@ bool FIKRetargetPelvicMotionOp::Initialize(
 	if (SourcePelvicBone == NAME_None || ApplyBodyMap(SourcePelvicBone) != Settings.TargetPelvicBone)
 	{
 		InLog.LogWarning( FText::Format(
-		LOCTEXT("InvalidPelvicBone", "{0}: Valid source/target mapping for pelvic bone must be specified. "), FText::FromName(GetName())));
+		LOCTEXT("InvalidPelvicMapping", "{0}: Valid source/target mapping for pelvic bone must be specified. "), FText::FromName(GetName())));
 		return false;
 	}
 	
