RelativeBodyOps: Add BPFL to help with bulk retarget -> export of sequences for comparison.

#Synced-CL 50425986

diff --git a/Engine/Plugins/Experimental/Animation/RelativeIKOp/Source/RelativeBodyAnimUtils/Private/RelativeIKBulkExportHelper.cpp b/Engine/Plugins/Experimental/Animation/RelativeIKOp/Source/RelativeBodyAnimUtils/Private/RelativeIKBulkExportHelper.cpp
new file mode 100644
--- /dev/null
+++ b/Engine/Plugins/Experimental/Animation/RelativeIKOp/Source/RelativeBodyAnimUtils/Private/RelativeIKBulkExportHelper.cpp
@@ -0,0 +1,47 @@
+﻿// Copyright Epic Games, Inc. All Rights Reserved.
+#include "RelativeIKBulkExportHelper.h"
+
+#include "Animation/AnimSequence.h"
+#include "Retargeter/IKRetargeter.h"
+#include "RetargetEditor/IKRetargetBatchOperation.h"
+
+#include UE_INLINE_GENERATED_CPP_BY_NAME(RelativeIKBulkExportHelper)
+
+void URelativeIKBulkExportHelper::BulkRetargetSequences(
+	UIKRetargeter* RetargetAsset,
+	UPARAM(ref) const TArray<UAnimSequence*>& RetargetSequenceList,
+	const FString& OutPath,
+	USkeletalMesh* SourceMesh,
+	USkeletalMesh* TargetMesh)
+{
+	if (RetargetSequenceList.IsEmpty())
+	{
+		return;
+	}
+	
+	FIKRetargetBatchOperationContext BatchContext;
+	BatchContext.NameRule.FolderPath = OutPath;
+	BatchContext.SourceMesh = SourceMesh;
+	BatchContext.TargetMesh = TargetMesh;
+	BatchContext.IKRetargetAsset = RetargetAsset;
+
+	// BatchContext.NameRule.Prefix = Prefix;
+	// BatchContext.NameRule.Suffix = TEXT("");
+	// BatchContext.NameRule.ReplaceFrom = SearchRM;
+	// BatchContext.NameRule.ReplaceTo = TEXT("");
+
+	BatchContext.bOverwriteExistingFiles = true;
+
+	for (const TObjectPtr<UAnimSequence> Seq : RetargetSequenceList)
+	{
+		if (!Seq)
+		{
+			continue;
+		}
+		
+		BatchContext.AssetsToRetarget.Add(Seq);
+	}
+
+	const TStrongObjectPtr<UIKRetargetBatchOperation> BatchOperation(NewObject<UIKRetargetBatchOperation>());
+	BatchOperation->RunRetarget(BatchContext);
+}
diff --git a/Engine/Plugins/Experimental/Animation/RelativeIKOp/Source/RelativeBodyAnimUtils/Public/RelativeBodyAnimModifier.h b/Engine/Plugins/Experimental/Animation/RelativeIKOp/Source/RelativeBodyAnimUtils/Public/RelativeBodyAnimModifier.h
--- a/Engine/Plugins/Experimental/Animation/RelativeIKOp/Source/RelativeBodyAnimUtils/Public/RelativeBodyAnimModifier.h
+++ b/Engine/Plugins/Experimental/Animation/RelativeIKOp/Source/RelativeBodyAnimUtils/Public/RelativeBodyAnimModifier.h
@@ -4,7 +4,6 @@
 
 #include "AnimationModifier.h"
 #include "MeshDescription.h"
-#include "Kismet/BlueprintFunctionLibrary.h"
 #include "UObject/ObjectPtr.h"
 
 #include "RelativeBodyAnimModifier.generated.h"
diff --git a/Engine/Plugins/Experimental/Animation/RelativeIKOp/Source/RelativeBodyAnimUtils/Public/RelativeIKBulkExportHelper.h b/Engine/Plugins/Experimental/Animation/RelativeIKOp/Source/RelativeBodyAnimUtils/Public/RelativeIKBulkExportHelper.h
new file mode 100644
--- /dev/null
+++ b/Engine/Plugins/Experimental/Animation/RelativeIKOp/Source/RelativeBodyAnimUtils/Public/RelativeIKBulkExportHelper.h
@@ -0,0 +1,27 @@
+﻿// Copyright Epic Games, Inc. All Rights Reserved.
+
+#pragma once
+#include "Kismet/BlueprintFunctionLibrary.h"
+#include "UObject/ObjectMacros.h"
+
+#include "RelativeIKBulkExportHelper.generated.h"
+
+class UIKRetargeter;
+class UAnimSequence;
+
+/**
+ * BPFL helper for running bulk anim retarget for export
+ */
+UCLASS(Transient, meta=(ScriptName="RelativeBodyBulkExportHelper"))
+class RELATIVEBODYANIMUTILS_API URelativeIKBulkExportHelper : public UBlueprintFunctionLibrary
+{
+	GENERATED_BODY()
+public:
+	UFUNCTION(BlueprintCallable, Category = "Export Tools")
+	static void BulkRetargetSequences(
+		UIKRetargeter* RetargetAsset,
+		UPARAM(ref) const TArray<UAnimSequence*>& RetargetSequenceList,
+		const FString& OutPath,
+		USkeletalMesh* SourceMesh = nullptr,
+		USkeletalMesh* TargetMesh = nullptr);
+};
diff --git a/Engine/Plugins/Experimental/Animation/RelativeIKOp/Source/RelativeBodyAnimUtils/RelativeBodyAnimUtils.Build.cs b/Engine/Plugins/Experimental/Animation/RelativeIKOp/Source/RelativeBodyAnimUtils/RelativeBodyAnimUtils.Build.cs
--- a/Engine/Plugins/Experimental/Animation/RelativeIKOp/Source/RelativeBodyAnimUtils/RelativeBodyAnimUtils.Build.cs
+++ b/Engine/Plugins/Experimental/Animation/RelativeIKOp/Source/RelativeBodyAnimUtils/RelativeBodyAnimUtils.Build.cs
@@ -42,6 +42,8 @@ namespace UnrealBuildTool.Rules
 				"AnimationCore",
 				"GeometryCore",
 				"DynamicMesh",
+				"IKRig",
+				"IKRigEditor",
 				"MeshConversion",
 				"MeshDescription",
 				"StaticMeshDescription",
