RelativeBodyModifiers: Add bp functions to set/get prop attach bones (bone refs inaccessible form bp).

#Synced-CL 50489965

diff --git a/Engine/Plugins/Experimental/Animation/RelativeIKOp/Source/RelativeBodyAnimUtils/Private/RelativeBodyBlueprintFunctions.cpp b/Engine/Plugins/Experimental/Animation/RelativeIKOp/Source/RelativeBodyAnimUtils/Private/RelativeBodyBlueprintFunctions.cpp
--- a/Engine/Plugins/Experimental/Animation/RelativeIKOp/Source/RelativeBodyAnimUtils/Private/RelativeBodyBlueprintFunctions.cpp
+++ b/Engine/Plugins/Experimental/Animation/RelativeIKOp/Source/RelativeBodyAnimUtils/Private/RelativeBodyBlueprintFunctions.cpp
@@ -38,6 +38,16 @@ TRelativeBodyClass* FindExactTypeBodyModifier(UAnimationModifiersAssetUserData*
 	return nullptr;
 }
 
+void URelativeBodyAnimBlueprintFunctions::SetPropAttachBone(UPARAM(ref) FPropsInfo& Info, FName AttachBoneName)
+{
+	Info.AttachBone = AttachBoneName;
+}
+
+FName URelativeBodyAnimBlueprintFunctions::GetPropAttachBone(const FPropsInfo& Info)
+{
+	return Info.AttachBone.BoneName;
+}
+
 void URelativeBodyAnimBlueprintFunctions::BulkUpdateRelativeBodyModifiers(TArray<UAnimSequenceBase*>& AnimationList, const FRelativeBodyAnimModifierOptions& Options)
 {
 	FScopedSlowTask ApplyModifiersTask(static_cast<float>(AnimationList.Num()), LOCTEXT("BulkUpdateProgress", "Updating and applying modifier to animation list"));
diff --git a/Engine/Plugins/Experimental/Animation/RelativeIKOp/Source/RelativeBodyAnimUtils/Public/RelativeBodyBlueprintFunctions.h b/Engine/Plugins/Experimental/Animation/RelativeIKOp/Source/RelativeBodyAnimUtils/Public/RelativeBodyBlueprintFunctions.h
--- a/Engine/Plugins/Experimental/Animation/RelativeIKOp/Source/RelativeBodyAnimUtils/Public/RelativeBodyBlueprintFunctions.h
+++ b/Engine/Plugins/Experimental/Animation/RelativeIKOp/Source/RelativeBodyAnimUtils/Public/RelativeBodyBlueprintFunctions.h
@@ -68,6 +68,12 @@ class RELATIVEBODYANIMUTILS_API URelativeBodyAnimBlueprintFunctions : public UBl
 	GENERATED_BODY()
 public:
 	
+	UFUNCTION(BlueprintCallable, Category = "Anim Modifier")
+	static void SetPropAttachBone(UPARAM(ref) FPropsInfo& Info, FName AttachBoneName);
+	
+	UFUNCTION(BlueprintCallable, Category = "Anim Modifier")
+	static FName GetPropAttachBone(const FPropsInfo& Info);
+	
 	UFUNCTION(BlueprintCallable, Category = "Anim Modifier")
 	static void BulkUpdateRelativeBodyModifiers(UPARAM(ref) TArray<UAnimSequenceBase*>& AnimationList, const FRelativeBodyAnimModifierOptions& Options);
 	
