RelativeIK: Prop Scalar for Optimization

#Synced-CL 50192622

diff --git a/Engine/Plugins/Experimental/Animation/RelativeIKOp/Source/RelativeIKOp/Private/RelativeIKOp.cpp b/Engine/Plugins/Experimental/Animation/RelativeIKOp/Source/RelativeIKOp/Private/RelativeIKOp.cpp
--- a/Engine/Plugins/Experimental/Animation/RelativeIKOp/Source/RelativeIKOp/Private/RelativeIKOp.cpp
+++ b/Engine/Plugins/Experimental/Animation/RelativeIKOp/Source/RelativeIKOp/Private/RelativeIKOp.cpp
@@ -189,7 +189,14 @@ void FIKRetargetRelativeIKOp::Run(
 	// Update bone goals to relative IK positions
 	if (Settings.bMultiBoneSolve)
 	{
-		ApplyMatrixBodyPairTargets(IKGoals, OutTargetGlobalPose, LocalDebugInfo, LocalTargetPairSpaces);
+		if (Settings.bSolveWithPropScale)
+		{
+			ApplyMatrixBodyPairTargetsWithPropScale(IKGoals, OutTargetGlobalPose, LocalDebugInfo, LocalTargetPairSpaces);
+		}
+		else
+		{
+			ApplyMatrixBodyPairTargets(IKGoals, OutTargetGlobalPose, LocalDebugInfo, LocalTargetPairSpaces);
+		}
 	}
 	else
 	{
@@ -220,6 +227,7 @@ void FIKRetargetRelativeIKOp::Run(
 			}
 		}
 	}
+	
 	UpdateDebugInfo(SourceScale, InSourceGlobalPose, OutTargetGlobalPose, LocalDebugInfo);
 #endif
 }
@@ -253,6 +261,7 @@ void FIKRetargetRelativeIKOp::RunAfterParent(
 		else
 		{
 			OutTargetGlobalPose[TargetBoneIdx] = ComputePinBoneTransform(PinData, InProcessor, InSourceGlobalPose, OutTargetGlobalPose);
+			OutTargetGlobalPose[TargetBoneIdx].SetScale3D(FVector(PinData.PropScalar));
 		}
 	}
 }
@@ -382,11 +391,15 @@ void FIKRetargetRelativeIKOp::ComputeFramePairSpaces(
 		FVector WeightedTargetOffset_O = TargetOffsetVec * (Settings.bMultiBoneSolve? 0.5*(FeasibilityWeight_O+FeasibilityWeight_P): FeasibilityWeight_O);
 		FVector WeightedTargetOffset_P = -TargetOffsetVec * (Settings.bMultiBoneSolve? 0.5*(FeasibilityWeight_O+FeasibilityWeight_P): FeasibilityWeight_P);
 		
+		FrameBoneConstraints[2*OutPairIdx].ContactOffset = TargetSpring_P-TargetPrimaryBoneLoc;
+		FrameBoneConstraints[2*OutPairIdx+1].ContactOffset = TargetSpring_O-TargetOtherBoneLoc;
 		FrameBoneConstraints[2*OutPairIdx].Offset = WeightedTargetOffset_P;
 		FrameBoneConstraints[2*OutPairIdx+1].Offset = WeightedTargetOffset_O;
+		FrameBoneConstraints[2*OutPairIdx].FeasibilityWeight = FeasibilityWeight_P;
+		FrameBoneConstraints[2*OutPairIdx+1].FeasibilityWeight = FeasibilityWeight_O;
 		FrameBoneConstraints[2*OutPairIdx].Weight = Weight;
 		FrameBoneConstraints[2*OutPairIdx+1].Weight = Weight;
-
+		
 #if WITH_EDITOR
 		// Debug draw verts
 		if (Settings.bDebugDraw)
@@ -483,10 +496,12 @@ void FIKRetargetRelativeIKOp::ComputeFramePropSpaces(
 			double FloorOffsetZMin = Settings.DistHalfPointLambda*Settings.FloorThreshold-TargetPropCompZMin.Z;
 			FVector TargetOffsetZMin = SourcePropCompZMin - TargetPropCompZMin;
 			FVector TargetOffsetZMax = SourcePropCompZMax - TargetPropCompZMax;
-			double Weight = Settings.FloorConstraintScaler*(SourceWeight*SourceWeight + (1-SourceWeight)*TargetWeight);
+			double Weight = Settings.FloorConstraintScalar*(SourceWeight*SourceWeight + (1-SourceWeight)*TargetWeight);
 
 			FrameBoneConstraints[2*OutPairIdx].ConstraintType = ERelativeIKFrameBoneConstraintType::Z;
 			FrameBoneConstraints[2*OutPairIdx+1].ConstraintType = ERelativeIKFrameBoneConstraintType::Z;
+			FrameBoneConstraints[2*OutPairIdx].ContactOffset = {0,0,(TargetPropCompZMin.Z - TargetPose[TargetPropPoseIdx].GetLocation().Z)};
+			FrameBoneConstraints[2*OutPairIdx+1].ContactOffset = {0,0,(TargetPropCompZMax.Z - TargetPose[TargetPropPoseIdx].GetLocation().Z)};
 			FrameBoneConstraints[2*OutPairIdx].Offset = {0,0,SourceWeight*TargetOffsetZMin.Z + (1-SourceWeight)*FloorOffsetZMin};
 			FrameBoneConstraints[2*OutPairIdx+1].Offset = {0,0,SourceWeight*TargetOffsetZMax.Z + (1-SourceWeight)*FloorOffsetZMin};
 			FrameBoneConstraints[2*OutPairIdx].Weight = Weight;
@@ -596,8 +611,12 @@ void FIKRetargetRelativeIKOp::ComputeFramePropSpaces(
 		FVector WeightedTargetOffset_B = TargetOffsetVec * (Settings.bMultiBoneSolve? 0.5*(FeasibilityWeight_B+FeasibilityWeight_P): FeasibilityWeight_B);
 		FVector WeightedTargetOffset_P = -TargetOffsetVec * (Settings.bMultiBoneSolve? 0.5*(FeasibilityWeight_B+FeasibilityWeight_P): FeasibilityWeight_P);
 
+		FrameBoneConstraints[2*OutPairIdx].ContactOffset = TargetSpring_P - TargetPropBoneLoc;
+		FrameBoneConstraints[2*OutPairIdx+1].ContactOffset = TargetSpring_B - TargetBodyBoneLoc;
 		FrameBoneConstraints[2*OutPairIdx].Offset = WeightedTargetOffset_P;
 		FrameBoneConstraints[2*OutPairIdx+1].Offset = WeightedTargetOffset_B;
+		FrameBoneConstraints[2*OutPairIdx].FeasibilityWeight = FeasibilityWeight_P;
+		FrameBoneConstraints[2*OutPairIdx+1].FeasibilityWeight = FeasibilityWeight_B;
 		FrameBoneConstraints[2*OutPairIdx].Weight = Weight;
 		FrameBoneConstraints[2*OutPairIdx+1].Weight = Weight;
 
@@ -633,8 +652,9 @@ void FIKRetargetRelativeIKOp::InitializePinData(const FIKRetargetProcessor& Proc
 	const TArray<FTransform>& TargetRetargeRefPose = TargetSkeleton.RetargetPoses.GetGlobalRetargetPose();
 	
 	CachePinBoneData.Reset(Settings.PinPropBones.Num());
-	for (const FPinBoneSettings& PinBoneInfo : Settings.PinPropBones)
+	for (int32 SettingId = 0; SettingId < Settings.PinPropBones.Num(); ++SettingId)
 	{
+		const FPinBoneSettings& PinBoneInfo = Settings.PinPropBones[SettingId];
 		FName SourceBoneName = PinBoneInfo.SourceBoneName;
 		FName TargetBoneName = ApplyBodyMap(SourceBoneName);
 		if (TargetBoneName == NAME_None)
@@ -666,6 +686,7 @@ void FIKRetargetRelativeIKOp::InitializePinData(const FIKRetargetProcessor& Proc
 		
 		PinData.SourceBone = SourceBoneName;
 		PinData.TargetBone = TargetBoneName;
+		PinData.PropScalar = PinBoneInfo.PropScalar;
 		
 		PinData.SourceTargetOffsetTfm = PropBoneRefPoseGlobal.GetRelativeTransform(TargetBoneRefPoseGlobal);
 		PinData.SourceLocalRefTfm = PropBoneParentRefPoseGlobal.GetRelativeTransform(PropBoneRefPoseGlobal);
@@ -682,7 +703,7 @@ void FIKRetargetRelativeIKOp::CachePinBoneTransforms(
 	TArray<FTransform>& OutTargetGlobalPose)
 {
 	CachePinTargetPropTransforms.SetNumZeroed(CachePinSourcePropBones.Num());
-	for (const FRelativeIKPinBoneData& PinData : CachePinBoneData)
+	for (FRelativeIKPinBoneData& PinData : CachePinBoneData)
 	{
 		FName SourcePropBone = PinData.SourceBone;
 		// Will calculate non-cached pins in ApplyAfterParent
@@ -700,7 +721,11 @@ void FIKRetargetRelativeIKOp::CachePinBoneTransforms(
 		
 		int32 TargetBoneIdx = CacheTargetSkelIndices[TargetPropBone];
 		FTransform OutGlobalTransform = ComputePinBoneTransform(PinData, Processor, InSourceGlobalPose, OutTargetGlobalPose);
-		
+		if (!Settings.bSolveWithPropScale)
+		{
+			PinData.PropScalar = Settings.PinPropBones[CachePropId2PinTargetPropSettingId[CachePropIdx]].PropScalar;
+		}
+		OutGlobalTransform.SetScale3D(FVector(PinData.PropScalar));
 		OutTargetGlobalPose[TargetBoneIdx] = OutGlobalTransform;
 		CachePinTargetPropTransforms[CachePropIdx] = OutGlobalTransform;
 	}
@@ -861,7 +886,6 @@ void FIKRetargetRelativeIKOp::ApplyAverageBodyPairTargets(
 	}
 }
 
-
 void FIKRetargetRelativeIKOp::ApplyMatrixBodyPairTargets(
 	TArray<FIKRigGoal>& OutIKGoals,
 	const TArray<FTransform>& TargetPose,
@@ -980,6 +1004,282 @@ void FIKRetargetRelativeIKOp::ApplyMatrixBodyPairTargets(
 	}
 }
 
+void FIKRetargetRelativeIKOp::ApplyMatrixBodyPairTargetsWithPropScale(
+	TArray<FIKRigGoal>& OutIKGoals,
+	const TArray<FTransform>& TargetPose,
+	FDebugRelativeIKDrawInfo& DebugInfo,
+	TArray<FDebugRelativeTargetPairSpace>& DebugPairSpaces)
+{
+	TArray<int32> OutIndexMap;
+	TArray<FName> OutBoneNames;
+	TArray<FVector> OutBoneLocations;
+	TArray<int32> ScalarIndexMap;
+	TArray<FName> ScalarBoneNames;
+
+	ScalarIndexMap.Reserve(CachePinSourcePropBones.Num());
+	ScalarBoneNames.Reserve(CachePinSourcePropBones.Num());
+	OutIndexMap.Reserve(CacheSourceBodyEffectVertIdx.Num());
+	OutBoneNames.Reserve(CacheSourceBodyEffectVertIdx.Num());
+	OutBoneLocations.Reserve(CacheSourceBodyEffectVertIdx.Num());
+	for (int EffectIdx = 0; EffectIdx < CacheSourceBodyEffectVertIdx.Num(); ++EffectIdx)
+	{
+		TArray<int32>& EffectIndices = CacheSourceBodyEffectVertIdx[EffectIdx];
+		if (EffectIndices.IsEmpty())
+		{
+			continue;
+		}
+
+		FName SourceBoneName = CacheSourceEffectBones[EffectIdx];
+		FName TargetBoneName = ApplyBodyMap(SourceBoneName);
+		if (!CacheTargetSkelIndices.Contains(TargetBoneName))
+		{
+			continue;
+		}
+
+		int32 TargetBoneIdx = CacheTargetSkelIndices[TargetBoneName];
+		const FTransform& TargetBonePose = TargetPose[TargetBoneIdx];
+		
+		int32 OutBoneIdx = OutIndexMap.Num();
+		OutIndexMap.Add(EffectIdx);
+		OutBoneNames.Add(SourceBoneName);
+		OutBoneLocations.Add(TargetBonePose.GetLocation());
+		
+		int32 PinPropIdx = CachePinSourcePropBones.Find(SourceBoneName);
+		if (PinPropIdx != INDEX_NONE)
+		{
+			ScalarIndexMap.Add(OutBoneIdx);
+	        ScalarBoneNames.Add(SourceBoneName);
+		}
+	}
+
+	Eigen::MatrixXd WeightMatrix = Eigen::MatrixXd::Zero(3*OutIndexMap.Num()+ScalarIndexMap.Num(), 3*OutIndexMap.Num()+ScalarIndexMap.Num());
+	Eigen::MatrixXd OffsetMatrix = Eigen::MatrixXd::Zero(3*OutIndexMap.Num()+ScalarIndexMap.Num(), 1);
+	// Eigen::VectorXd BaseWeights = Eigen::VectorXd::Zero(OutIndexMap.Num());
+
+	for (int32 RowIdx = 0; RowIdx < OutIndexMap.Num(); ++RowIdx)
+	{
+		int32 EffectIdx = OutIndexMap[RowIdx];
+		TArray<int32>& EffectIndices = CacheSourceBodyEffectVertIdx[EffectIdx];
+		FVector OffsetLoc = FVector::ZeroVector;
+		for (int Index : EffectIndices)
+		{
+			const FRelativeIKFrameBoneConstraint& BoneConstraint = FrameBoneConstraints[Index];
+			if (FMath::IsNearlyZero(BoneConstraint.Weight))
+			{
+				continue;
+			}
+
+			FName EffectBone = CachePairEffectBoneNames[Index];
+			int32 OutWeightIdx = OutBoneNames.Find(EffectBone);
+
+			if (BoneConstraint.ConstraintType == ERelativeIKFrameBoneConstraintType::XYZ)
+			{
+				WeightMatrix(3*RowIdx+0, 3*RowIdx+0) += BoneConstraint.Weight;
+				WeightMatrix(3*RowIdx+1, 3*RowIdx+1) += BoneConstraint.Weight;
+				WeightMatrix(3*RowIdx+2, 3*RowIdx+2) += BoneConstraint.Weight;
+				if (OutWeightIdx != INDEX_NONE)
+				{
+					WeightMatrix(3*RowIdx+0, 3*OutWeightIdx+0) -= BoneConstraint.Weight;
+					WeightMatrix(3*RowIdx+1, 3*OutWeightIdx+1) -= BoneConstraint.Weight;
+					WeightMatrix(3*RowIdx+2, 3*OutWeightIdx+2) -= BoneConstraint.Weight;
+				}
+			}
+			else if (BoneConstraint.ConstraintType == ERelativeIKFrameBoneConstraintType::Z)
+			{
+				WeightMatrix(3*RowIdx+2, 3*RowIdx+2) += BoneConstraint.Weight;
+				if (OutWeightIdx != INDEX_NONE)
+				{
+					WeightMatrix(3*RowIdx+2, 3*OutWeightIdx+2) -= BoneConstraint.Weight;
+				}
+			}
+			
+			OffsetLoc += BoneConstraint.Weight * BoneConstraint.Offset;
+		}
+		
+		OffsetMatrix(3*RowIdx+0) = OffsetLoc.X;
+		OffsetMatrix(3*RowIdx+1) = OffsetLoc.Y;
+		OffsetMatrix(3*RowIdx+2) = OffsetLoc.Z;
+	}
+	WeightMatrix.diagonal().array() += Settings.MultiBoneRegularization;
+
+	for (int32 ScalarIdx = 0; ScalarIdx < ScalarIndexMap.Num(); ++ScalarIdx)
+	{
+		int32 RowIdx = 3*OutIndexMap.Num() + ScalarIdx;
+		WeightMatrix(RowIdx, RowIdx) = Settings.PropScalarDamping;
+		int32 OutBoneIdx = ScalarIndexMap[ScalarIdx];
+		WeightMatrix(3*OutBoneIdx+0, 3*OutBoneIdx+0) += Settings.PropScalarBoneRegularization-Settings.MultiBoneRegularization;
+		WeightMatrix(3*OutBoneIdx+1, 3*OutBoneIdx+1) += Settings.PropScalarBoneRegularization-Settings.MultiBoneRegularization;
+		WeightMatrix(3*OutBoneIdx+2, 3*OutBoneIdx+2) += Settings.PropScalarBoneRegularization-Settings.MultiBoneRegularization;
+		int32 EffectIdx = OutIndexMap[OutBoneIdx];
+		TArray<int32>& EffectIndices = CacheSourceBodyEffectVertIdx[EffectIdx];
+		FVector ScalarOffsetLoc = FVector::ZeroVector;
+		
+		int32 CachePropIdx = CachePinSourcePropBones.Find(ScalarBoneNames[ScalarIdx]);
+		double PropScalar = CachePinBoneData[CachePropId2PinBoneDataId[CachePropIdx]].PropScalar;
+		
+		for (int Index : EffectIndices)
+		{
+			const FRelativeIKFrameBoneConstraint& BoneConstraint = FrameBoneConstraints[Index];
+			if (FMath::IsNearlyZero(BoneConstraint.Weight))
+			{
+				continue;
+			}
+			int IndexOpp;
+			if (Index%2==0)
+			{
+				IndexOpp = Index + 1;
+			}
+			else
+			{
+				IndexOpp = Index - 1;
+			}
+			const FRelativeIKFrameBoneConstraint& BoneConstraintOpp = FrameBoneConstraints[IndexOpp];
+			
+			FName EffectBone = CachePairEffectBoneNames[Index];
+			int32 OutWeightIdx = OutBoneNames.Find(EffectBone);
+			int32 OutScaleIdx = ScalarBoneNames.Find(EffectBone);
+			
+			FVector ContactOffset = FVector::ZeroVector;
+			if (!FMath::IsNearlyZero(PropScalar))
+			{
+				ContactOffset = BoneConstraint.ContactOffset/PropScalar;
+			}
+			
+			if (BoneConstraint.ConstraintType == ERelativeIKFrameBoneConstraintType::XYZ)
+			{
+				
+				WeightMatrix(RowIdx, RowIdx) += BoneConstraint.Weight * ContactOffset.X * ContactOffset.X;
+				WeightMatrix(RowIdx, 3*OutBoneIdx+0) += BoneConstraint.Weight * ContactOffset.X;
+				WeightMatrix(RowIdx, RowIdx) += BoneConstraint.Weight * ContactOffset.Y * ContactOffset.Y;
+				WeightMatrix(RowIdx, 3*OutBoneIdx+1) += BoneConstraint.Weight * ContactOffset.Y;
+				WeightMatrix(RowIdx, RowIdx) += BoneConstraint.Weight * ContactOffset.Z * ContactOffset.Z;
+				WeightMatrix(RowIdx, 3*OutBoneIdx+2) += BoneConstraint.Weight * ContactOffset.Z;
+				if (OutWeightIdx != INDEX_NONE)
+				{
+					WeightMatrix(RowIdx, 3*OutWeightIdx+0) -= BoneConstraint.Weight * ContactOffset.X;
+					WeightMatrix(RowIdx, 3*OutWeightIdx+1) -= BoneConstraint.Weight * ContactOffset.Y;
+					WeightMatrix(RowIdx, 3*OutWeightIdx+2) -= BoneConstraint.Weight * ContactOffset.Z;
+				}
+				if (OutScaleIdx != INDEX_NONE)
+				{
+					int32 OutCachePropIdx = CachePinSourcePropBones.Find(ScalarBoneNames[OutScaleIdx]);
+					double OutPropScalar = CachePinBoneData[CachePropId2PinBoneDataId[OutCachePropIdx]].PropScalar;
+					OutScaleIdx = 3*OutIndexMap.Num() + OutScaleIdx;
+					if (!FMath::IsNearlyZero(OutPropScalar))
+					{
+						WeightMatrix(RowIdx, OutScaleIdx) -= 0.5*BoneConstraint.Weight * ContactOffset.X * BoneConstraintOpp.ContactOffset.X/OutPropScalar;
+						WeightMatrix(RowIdx, OutScaleIdx) -= 0.5*BoneConstraint.Weight * ContactOffset.Y * BoneConstraintOpp.ContactOffset.Y/OutPropScalar;
+						WeightMatrix(RowIdx, OutScaleIdx) -= 0.5*BoneConstraint.Weight * ContactOffset.Z * BoneConstraintOpp.ContactOffset.Z/OutPropScalar;
+					}
+				}
+			}
+			else if (BoneConstraint.ConstraintType == ERelativeIKFrameBoneConstraintType::Z)
+			{
+				WeightMatrix(RowIdx, RowIdx) += BoneConstraint.Weight * ContactOffset.Z * ContactOffset.Z;
+				WeightMatrix(RowIdx, 3*OutBoneIdx+2) += BoneConstraint.Weight * ContactOffset.Z;
+				if (OutWeightIdx != INDEX_NONE)
+				{
+					WeightMatrix(RowIdx, 3*OutWeightIdx+2) -= BoneConstraint.Weight * ContactOffset.Z;
+				}
+				if (OutScaleIdx != INDEX_NONE)
+				{
+					int32 OutCachePropIdx = CachePinSourcePropBones.Find(ScalarBoneNames[OutScaleIdx]);
+					double OutPropScalar = CachePinBoneData[CachePropId2PinBoneDataId[OutCachePropIdx]].PropScalar;
+					OutScaleIdx = 3*OutIndexMap.Num() + OutScaleIdx;
+					if (!FMath::IsNearlyZero(OutPropScalar))
+					{
+						WeightMatrix(RowIdx, OutScaleIdx) -= 0.5*BoneConstraint.Weight * ContactOffset.Z * BoneConstraintOpp.ContactOffset.Z/OutPropScalar;
+					}
+				}	
+			}
+
+			ScalarOffsetLoc += BoneConstraint.Weight * BoneConstraint.Offset * ContactOffset;
+		}
+		
+		OffsetMatrix(RowIdx) += ScalarOffsetLoc.X + ScalarOffsetLoc.Y + ScalarOffsetLoc.Z + Settings.PropScalarDamping*(Settings.PinPropBones[CachePropId2PinTargetPropSettingId[CachePropIdx]].PropScalar-PropScalar);
+	}
+
+	WeightMatrix = WeightMatrix.selfadjointView<Eigen::Lower>();
+
+	Eigen::MatrixXd SolvedOffset(3*OutIndexMap.Num() + ScalarIndexMap.Num(), 1);
+	Eigen::LLT<Eigen::MatrixXd> MatrixSolver = WeightMatrix.llt();
+	SolvedOffset = MatrixSolver.solve(OffsetMatrix);
+
+	for (int32 ScalarIdx = 0; ScalarIdx < ScalarIndexMap.Num(); ++ScalarIdx)
+	{
+		int32 RowIdx = 3*OutIndexMap.Num() + ScalarIdx;
+		int32 OutBoneIdx = ScalarIndexMap[ScalarIdx];
+		int32 EffectIdx = OutIndexMap[OutBoneIdx];
+		if (FMath::IsNearlyZero(WeightMatrix(RowIdx, RowIdx)))
+		{
+			continue;
+		}
+		
+		FName SourceBoneName = CacheSourceEffectBones[EffectIdx];
+		FName TargetBoneName = ApplyBodyMap(SourceBoneName);
+		if (!CacheTargetSkelIndices.Contains(TargetBoneName))
+		{
+			continue;
+		}
+		
+		int32 CachePropIdx = CachePinSourcePropBones.Find(SourceBoneName);
+		if (CachePropIdx == INDEX_NONE)
+		{
+			continue;
+		}
+		
+		double &PropScalar = CachePinBoneData[CachePropId2PinBoneDataId[CachePropIdx]].PropScalar;
+		if (Settings.bDebugDraw && Settings.bDryRun)
+		{
+			PropScalar = Settings.PinPropBones[CachePropId2PinTargetPropSettingId[CachePropIdx]].PropScalar;
+			continue;
+		}
+		// PropScalar = FMath::Lerp(PropScalar, Settings.PinPropBones[CachePinTargetPropSettingId[CachePropIdx]].PropScalar,FMath::Clamp(Settings.PropScalarLearningRate*5,0,1));
+		PropScalar += FMath::Clamp(SolvedOffset(RowIdx),-Settings.PropScalarLearningRate*PropScalar, Settings.PropScalarLearningRate*PropScalar);
+		PropScalar = FMath::Clamp(PropScalar, Settings.PinPropBones[CachePropId2PinTargetPropSettingId[CachePropIdx]].PropScalar*(1-Settings.PropScalarRange), Settings.PinPropBones[CachePropId2PinTargetPropSettingId[CachePropIdx]].PropScalar*(1+Settings.PropScalarRange));
+		// PropScalar += SolvedOffset(RowIdx);
+		CachePinTargetPropTransforms[CachePropIdx].SetScale3D(FVector(PropScalar));
+		
+		WeightMatrix.row(RowIdx).setZero();
+		WeightMatrix.col(RowIdx).setZero();
+		WeightMatrix(RowIdx,RowIdx) = 1.0;
+		OffsetMatrix(RowIdx) = PropScalar-1;
+	}
+	
+	//Solve again for fixed prop scalar
+	MatrixSolver = WeightMatrix.llt();
+	SolvedOffset = MatrixSolver.solve(OffsetMatrix);
+
+	for (int32 RowIdx = 0; RowIdx < OutIndexMap.Num(); ++RowIdx)
+	{
+		int32 EffectIdx = OutIndexMap[RowIdx];
+		FVector TotalWeight = FVector{WeightMatrix(3*RowIdx+0, 3*RowIdx+0), WeightMatrix(3*RowIdx+1, 3*RowIdx+1), WeightMatrix(3*RowIdx+2, 3*RowIdx+2)};
+		if (FMath::IsNearlyZero(TotalWeight.GetMax()))
+		{
+			continue;
+		}
+		
+		FName SourceBoneName = CacheSourceEffectBones[EffectIdx];
+		FName TargetBoneName = ApplyBodyMap(SourceBoneName);
+		if (!CacheTargetSkelIndices.Contains(TargetBoneName))
+		{
+			continue;
+		}
+		
+		FVector TargetOffset{
+			(!FMath::IsNearlyZero(TotalWeight.X)) ? SolvedOffset(3*RowIdx+0) : 0.0,
+			(!FMath::IsNearlyZero(TotalWeight.Y)) ? SolvedOffset(3*RowIdx+1) : 0.0,
+			(!FMath::IsNearlyZero(TotalWeight.Z)) ? SolvedOffset(3*RowIdx+2) : 0.0
+		};
+
+		int32 TargetBoneIdx = CacheTargetSkelIndices[TargetBoneName];
+		const FTransform& TargetBonePose = TargetPose[TargetBoneIdx];
+		
+		UpdateTargetIKGoal(OutIKGoals, TargetBoneName, TotalWeight, TargetOffset, TargetBonePose);
+		UpdatePropPinTransform(SourceBoneName, TotalWeight, TargetOffset, TargetBonePose);
+	}
+}
+
 void FIKRetargetRelativeIKOp::UpdateTargetIKGoal(
 	TArray<FIKRigGoal>& OutIKGoals,
 	FName TargetBoneName,
@@ -1985,6 +2285,31 @@ void FIKRetargetRelativeIKOp::UpdateCachePropNotifyInfo(URelativePropsBakeAnimNo
 		// Keep track of prop attach bones for post-solve pin/update
 		CachePinSourcePropBones.AddUnique(PropBone);
 	}
+	
+	CachePropId2PinTargetPropSettingId.SetNumZeroed(CachePinSourcePropBones.Num());
+	CachePropId2PinBoneDataId.SetNumZeroed(CachePinSourcePropBones.Num());
+	int32 PinBoneId = 0;
+	for (int32 SettingId = 0; SettingId < Settings.PinPropBones.Num(); ++SettingId)
+	{
+		const FPinBoneSettings& PinBoneInfo = Settings.PinPropBones[SettingId];
+		FName SourceBoneName = PinBoneInfo.SourceBoneName;
+		FName TargetBoneName = ApplyBodyMap(SourceBoneName);
+		
+		int32 CachePropIdx = CachePinSourcePropBones.Find(SourceBoneName);
+		if (CachePropIdx != INDEX_NONE)
+		{
+			CachePropId2PinTargetPropSettingId[CachePropIdx] = SettingId;
+		}
+		
+		if (TargetBoneName != NAME_None && CacheSourceSkelIndices.Contains(SourceBoneName) && CacheTargetSkelIndices.Contains(TargetBoneName))
+		{
+			if (CachePropIdx != INDEX_NONE)
+			{
+				CachePropId2PinBoneDataId[CachePropIdx] = PinBoneId;
+			}
+			PinBoneId++;
+		}
+	}
 }
 
 
diff --git a/Engine/Plugins/Experimental/Animation/RelativeIKOp/Source/RelativeIKOp/Public/RelativeIKOp.h b/Engine/Plugins/Experimental/Animation/RelativeIKOp/Source/RelativeIKOp/Public/RelativeIKOp.h
--- a/Engine/Plugins/Experimental/Animation/RelativeIKOp/Source/RelativeIKOp/Public/RelativeIKOp.h
+++ b/Engine/Plugins/Experimental/Animation/RelativeIKOp/Source/RelativeIKOp/Public/RelativeIKOp.h
@@ -118,6 +118,10 @@ struct FPinBoneSettings
 	// The method used to calculate the rotation of the bone.
 	UPROPERTY(EditAnywhere, BlueprintReadWrite, Category = "Pinning", meta=(ReinitializeOnEdit))
 	EPinBoneRotationMode RotationMode = EPinBoneRotationMode::CopyGlobalRotation;
+	
+	// Alpha between primary and secondary pair distance relationship
+	UPROPERTY(EditAnywhere, BlueprintReadWrite, Category = "Pinning", meta=(ReinitializeOnEdit))
+	double PropScalar = 1.0;
 };
 
 UENUM()
@@ -133,6 +137,8 @@ struct FRelativeIKFrameBoneConstraint
 	ERelativeIKFrameBoneConstraintType ConstraintType;
 	
 	double Weight;
+	double FeasibilityWeight;
+	FVector ContactOffset;
 	FVector Offset;
 
 	FRelativeIKFrameBoneConstraint()
@@ -171,7 +177,7 @@ struct FIKRetargetRelativeIKOpSettings : public FIKRetargetOpSettingsBase
 	
 	// Maximum distance for which prop-floor info is baked
     UPROPERTY(EditAnywhere, BlueprintReadWrite, Category = "Parameters", meta=(ClampMin = "0.0"))
-    double FloorConstraintScaler = 10.00;
+    double FloorConstraintScalar = 10.00;
 	
 	// List of bones to pin. Should include all attach bones that may be affected by relative ik bake data 
 	UPROPERTY(EditAnywhere, BlueprintReadWrite, Category = "Prop Settings", meta=(ReinitializeOnEdit))
@@ -253,9 +259,31 @@ struct FIKRetargetRelativeIKOpSettings : public FIKRetargetOpSettingsBase
 	UPROPERTY(EditAnywhere, BlueprintReadWrite, Category = "Test Features")
 	bool bMultiBoneSolve = true;
 	
-	// Regularization additive constant for stabilizing multibone solve
-	UPROPERTY(EditAnywhere, BlueprintReadWrite, Category = "Test Features", meta=(ClampMin = "0", ClampMax = "1"))
-	double MultiBoneRegularization = 0.0;
+	// Regularization additive constant for stabilizing multibone solve, the higher the body bones will stay at original position
+	UPROPERTY(EditAnywhere, BlueprintReadWrite, Category = "Test Features", meta=(ClampMin = "0.001"))
+	double MultiBoneRegularization = 0.001;
+
+		
+	// Add Props scale as a parameter for optimization
+	UPROPERTY(EditAnywhere, BlueprintReadWrite, Category = "Prop Scaling")
+	bool bSolveWithPropScale = false;
+	
+	// Props scalar change ratio which allow to change per frame
+	UPROPERTY(EditAnywhere, BlueprintReadWrite, Category = "Prop Scaling", meta=(ClampMin = "0.0001", ClampMax = "10"))
+	double PropScalarLearningRate = 0.001;
+	
+	// Test Parameter for prop scalar range allow to change, it will clamp prop scalar between [PinPropBone.Propscalar-PropScalarRange, PinPropBone.Propscalar+PropScalarRange]
+	UPROPERTY(EditAnywhere, BlueprintReadWrite, Category = "Prop Scaling", meta=(ClampMin = "0", ClampMax = "10"))
+	double PropScalarRange = 0.25;
+	
+	// Regularization additive constant for stabilizing prop bone position, the higher the prop bones will stay at original position
+	UPROPERTY(EditAnywhere, BlueprintReadWrite, Category = "Prop Scaling", meta=(ClampMin = "0.001"))
+	double PropScalarBoneRegularization = 1;
+	
+	// Regularization additive constant for stabilizing prop bone scalar, the higher the prop bones will stay at original Propscalar set in PinPropBone Settings 
+	UPROPERTY(EditAnywhere, BlueprintReadWrite, Category = "Prop Scaling", meta=(ClampMin = "0.001"))
+	double PropScalarDamping = 500;
+
 	
 	// Distance blend body pairs towards each contact position if both can move
 	UPROPERTY(EditAnywhere, BlueprintReadWrite, Category = "Test Features")
@@ -336,6 +364,8 @@ private:
 		FTransform SourceLocalRefTfm;
 		FTransform TargetLocalRefTfm;
 		FQuat RefRotDelta;
+		
+		double PropScalar;
 	};
 	
 	void InitializePinData(const FIKRetargetProcessor& Processor);
@@ -382,6 +412,12 @@ private:
 		const TArray<FTransform>& TargetPose,
 		FDebugRelativeIKDrawInfo& DebugInfo,
 		TArray<FDebugRelativeTargetPairSpace>& DebugPairSpaces);
+	
+	void ApplyMatrixBodyPairTargetsWithPropScale(
+		TArray<FIKRigGoal>& OutIKGoals,
+		const TArray<FTransform>& TargetPose,
+		FDebugRelativeIKDrawInfo& DebugInfo,
+		TArray<FDebugRelativeTargetPairSpace>& DebugPairSpaces);
 
 	void UpdateTargetIKGoal(
 		TArray<FIKRigGoal>& OutIKGoals,
@@ -531,6 +567,8 @@ private:
 	TArray<FName> CachePinSourcePropBones;
 	TArray<FRelativeIKPinBoneData> CachePinBoneData;
 	TArray<FTransform> CachePinTargetPropTransforms;
+	TArray<int32> CachePropId2PinTargetPropSettingId;
+	TArray<int32> CachePropId2PinBoneDataId;
 
 	// Cached montage instance/segment info for early-out matching in preupdate
 	const FAnimMontageInstance* MontageInstance;
