RelativePropsAnimModifier: Warning for no props

#Synced-CL 50555092

diff --git a/Engine/Plugins/Experimental/Animation/RelativeIKOp/Source/RelativeBodyAnimUtils/Private/RelativePropsAnimModifier.cpp b/Engine/Plugins/Experimental/Animation/RelativeIKOp/Source/RelativeBodyAnimUtils/Private/RelativePropsAnimModifier.cpp
--- a/Engine/Plugins/Experimental/Animation/RelativeIKOp/Source/RelativeBodyAnimUtils/Private/RelativePropsAnimModifier.cpp
+++ b/Engine/Plugins/Experimental/Animation/RelativeIKOp/Source/RelativeBodyAnimUtils/Private/RelativePropsAnimModifier.cpp
@@ -223,11 +223,6 @@ void FPropsInfo::GetRefToAnimPoseMatrices(float AnimPlayhead, TArray<FMatrix44f>
 
 void URelativePropsAnimModifier::OnApply_Implementation(UAnimSequence* InAnimation)
 {
-	if (NotifyClass)
-	{
-		URelativeBodyAnimModifier::OnApply_Implementation(InAnimation);
-	}
-
 	// TODO: Make log category for relative body utils?
 	if (InAnimation == nullptr)
 	{
@@ -281,6 +276,45 @@ void URelativePropsAnimModifier::OnApply_Implementation(UAnimSequence* InAnimati
 	const int32 NumBodies = BodyIndicesToBoneIndices.Num();
 	const float HashGridCellSize = std::min(1.5f * ContactThreshold, 50.f); // the length of the cell size in the point hash gri
 
+	// Check for valid prop data
+	TArray<bool> IsValidPropData;
+	IsValidPropData.SetNumZeroed(PropsData.Num());
+	bool PropsDataRequired = false;
+	for (int32 PropsIdx = 0; PropsIdx < PropsData.Num(); PropsIdx++)
+	{
+		FMeshDescription* PropsMeshDescription = PropsData[PropsIdx].GetMeshDescription(0); // LOD 0
+		if (!PropsMeshDescription)
+		{
+			UE_LOG(LogAnimation, Warning, TEXT("RelativePropAnimModifier: Bad propdata at index [%d]. \n\tReason: PropsMeshDescription Error!"), PropsIdx);
+			continue;
+		}
+		
+		if (PropsData[PropsIdx].PropStaticMeshAsset && PropsData[PropsIdx].PropAnimSequence && !PropsData[PropsIdx].PropSkeletalMeshAsset)
+		{
+			UE_LOG(LogAnimation, Warning, TEXT("RelativePropAnimModifier: Bad propdata at index [%d]. \n\tReason: Need SkeletalMesh to support Prop Animation!"), PropsIdx);
+			continue;
+		}
+		
+		if (RefSkeleton.FindBoneIndex(PropsData[PropsIdx].AttachBone.BoneName) == INDEX_NONE)
+		{
+			UE_LOG(LogAnimation, Warning, TEXT("RelativePropAnimModifier: Bad propdata at index [%d]. \n\tReason: Attach bone '%s' not found in skeletal mesh '%s'!"), PropsIdx, *PropsData[PropsIdx].AttachBone.BoneName.ToString(), *GetNameSafe(SkeletalMeshAsset));
+			continue;
+		}
+
+		IsValidPropData[PropsIdx] = true;
+		PropsDataRequired = true;
+	}
+	
+	if (NotifyClass)
+	{
+		URelativeBodyAnimModifier::OnApply_Implementation(InAnimation);
+	}
+		
+	if (!PropsDataRequired)
+	{
+		return;
+	}
+	
 	// Process animation asset
 	{
 		const FAnimPoseEvaluationOptions AnimPoseEvalOptions{
@@ -296,9 +330,8 @@ void URelativePropsAnimModifier::OnApply_Implementation(UAnimSequence* InAnimati
 		const float SequenceLength = InAnimation->GetPlayLength();
 		const float SampleStep = 1.0f / FMath::Max(1.0f, static_cast<float>(SampleRate));
 		const int SampleNum = FMath::TruncToInt(SequenceLength / SampleStep);
-
+		
 		URelativePropsBakeAnimNotify* BakedAnimNotifyInfo = nullptr;
-
 		bool bUseDenseNotifyInfo = PropsNotifyClass->IsChildOf(URelativePropsBakeAnimNotify::StaticClass());
 		TArray<FPropsMeshAndTree> PropsTreeData;
 		PropsTreeData.SetNum(PropsData.Num());
@@ -340,12 +373,13 @@ void URelativePropsAnimModifier::OnApply_Implementation(UAnimSequence* InAnimati
 				
 			for (size_t PropsIdx = 0; PropsIdx < PropsData.Num(); PropsIdx++)
 			{
-				BakedAnimNotifyInfo->OffsetTransformsForBones.Emplace(PropsData[PropsIdx].AttachBone.BoneName,  PropsData[PropsIdx].AttachTransform);
-				FMeshDescription* PropsMeshDescription = PropsData[PropsIdx].GetMeshDescription(0); // LOD 0
-				if (PropsMeshDescription)
+				if (!IsValidPropData[PropsIdx])
 				{
-					PropsTreeData[PropsIdx].Init(PropsMeshDescription);
+					continue;
 				}
+				BakedAnimNotifyInfo->OffsetTransformsForBones.Emplace(PropsData[PropsIdx].AttachBone.BoneName,  PropsData[PropsIdx].AttachTransform);
+				FMeshDescription* PropsMeshDescription = PropsData[PropsIdx].GetMeshDescription(0); // LOD 0
+				PropsTreeData[PropsIdx].Init(PropsMeshDescription);
 	
 				for (int32 BodyIndex1 = 0; BodyIndex1 < NumBodies; ++BodyIndex1)
 				{
@@ -410,6 +444,11 @@ void URelativePropsAnimModifier::OnApply_Implementation(UAnimSequence* InAnimati
 				BakedAnimNotifyInfo->PropsPairsSampleTime[SampleIndex] = SampleTime;
 				for (int32 PropsIdx = 0; PropsIdx < PropsData.Num(); PropsIdx++)
 				{
+					if (!IsValidPropData[PropsIdx])
+					{
+						continue;
+					}
+					
 					FTransform AttachBoneTransform = UAnimPoseExtensions::GetBonePose(AnimPose, PropsData[PropsIdx].AttachBone.BoneName, EAnimPoseSpaces::World);
 					FTransform PropsLocalToWorld = PropsData[PropsIdx].AttachTransform*AttachBoneTransform;
 					const FPositionVertexBuffer& PositionVertexBuffer = PropsData[PropsIdx].GetPositionVertexBuffer();
