RelativePropsAnimModifier: Always make prop modifier if using bulk prop create functions. Also add new BP funciton to "Set Relative Anim Modifiers" without applying.

#Synced-CL 50609381

diff --git a/Engine/Plugins/Experimental/Animation/RelativeIKOp/Source/RelativeBodyAnimUtils/Private/RelativeBodyBlueprintFunctions.cpp b/Engine/Plugins/Experimental/Animation/RelativeIKOp/Source/RelativeBodyAnimUtils/Private/RelativeBodyBlueprintFunctions.cpp
--- a/Engine/Plugins/Experimental/Animation/RelativeIKOp/Source/RelativeBodyAnimUtils/Private/RelativeBodyBlueprintFunctions.cpp
+++ b/Engine/Plugins/Experimental/Animation/RelativeIKOp/Source/RelativeBodyAnimUtils/Private/RelativeBodyBlueprintFunctions.cpp
@@ -48,13 +48,15 @@
 	return Info.AttachBone.BoneName;
 }
 
-void URelativeBodyAnimBlueprintFunctions::BulkUpdateRelativeBodyModifiers(TArray<UAnimSequenceBase*>& AnimationList, const FRelativeBodyAnimModifierOptions& Options)
+float URelativeBodyAnimBlueprintFunctions::GetSegmentStartTime(const FAnimNotifyEvent& AnimNotifyEvent)
 {
-	FScopedSlowTask ApplyModifiersTask(static_cast<float>(AnimationList.Num()), LOCTEXT("BulkUpdateProgress", "Updating and applying modifier to animation list"));
-	ApplyModifiersTask.MakeDialog();
+	return AnimNotifyEvent.GetTime();
+}
+
+void URelativeBodyAnimBlueprintFunctions::SetRelativeAnimModifiers(UPARAM(ref) TArray<UAnimSequenceBase*>& AnimationList, const FPropBodyAnimModifierOptions& Options)
+{
 	for (UAnimSequenceBase* Anim : AnimationList)
 	{
-		ApplyModifiersTask.EnterProgressFrame(1.0f, FText::FromString(Anim->GetName()));
 		UAnimationModifiersAssetUserData* AssetUserData = Anim->GetAssetUserData<UAnimationModifiersAssetUserData>();
 		if (!AssetUserData)
 		{
@@ -62,49 +64,65 @@
 			continue;
 		}
 
-		TObjectPtr<URelativeBodyAnimModifier> UpdateModifierInst = FindRelativeBodyModifier<URelativeBodyAnimModifier>(AssetUserData);
+		TObjectPtr<URelativePropsAnimModifier> UpdateModifierInst = FindRelativeBodyModifier<URelativePropsAnimModifier>(AssetUserData);
 		if (!UpdateModifierInst)
 		{
-			if (!AssetUserData->AddAnimationModifierOfClass(Anim, URelativeBodyAnimModifier::StaticClass()))
+			if (!AssetUserData->AddAnimationModifierOfClass(Anim, URelativePropsAnimModifier::StaticClass()))
 			{
 				UE_LOG(LogAnimation, Warning, TEXT("Unable to create new modifier for anim: %s"), *Anim->GetName());
 				continue;
 			}
 
-			UpdateModifierInst = FindRelativeBodyModifier<URelativeBodyAnimModifier>(AssetUserData);
+			UpdateModifierInst = FindRelativeBodyModifier<URelativePropsAnimModifier>(AssetUserData);
 		}
 
 		if (!UpdateModifierInst)
 		{
 			continue;
 		}
-
+		
+		UpdateModifierInst->PropsData.Empty(Options.PropMeshSettings.Num());
+		for (const FPropsInfo& PropInfo : Options.PropMeshSettings)
+		{
+			UpdateModifierInst->PropsData.Add(PropInfo);
+		}
+		UpdateModifierInst->bPropsFloorInfoBaking = Options.bPropsFloorInfoBaking;
+		UpdateModifierInst->PropsNotifyClass = Options.PropNotifyClass;
+		
+		// Always set these
 		UpdateModifierInst->SampleRate = Options.SampleRate;
 		UpdateModifierInst->ContactThreshold = Options.ContactThreshold;
 		UpdateModifierInst->SkeletalMeshAsset = Options.SkeletalMeshAsset;
 		UpdateModifierInst->PhysicsAssetOverride = Options.PhysicsAssetOverride;
 		UpdateModifierInst->DomainBodyNames = Options.DomainBodyNames;
 		UpdateModifierInst->ContactBodyNames = Options.ContactBodyNames;
-		UpdateModifierInst->NotifyClass = Options.NotifyClass;
-
-		UAnimSequence* AnimSeq = Cast<UAnimSequence>(Anim);
-		if (AnimSeq)
+		UpdateModifierInst->NotifyClass = nullptr;
+		
+		// TODO: Ideally we would just delete the old body modifier in this case
+		bool bUpdateSeparateBodyModifier = false;
+		TObjectPtr<URelativeBodyAnimModifier> BodyModifierInst = FindExactTypeBodyModifier<URelativeBodyAnimModifier>(AssetUserData);
+		if (Options.NotifyClass != nullptr && BodyModifierInst != nullptr)
 		{
-			UpdateModifierInst->ApplyToAnimationSequence(AnimSeq);
+			// Update body modifier to generate the 
+			bUpdateSeparateBodyModifier = true;
+			BodyModifierInst->SampleRate = Options.SampleRate;
+			BodyModifierInst->ContactThreshold = Options.ContactThreshold;
+			BodyModifierInst->SkeletalMeshAsset = Options.SkeletalMeshAsset;
+			BodyModifierInst->PhysicsAssetOverride = Options.PhysicsAssetOverride;
+			BodyModifierInst->DomainBodyNames = Options.DomainBodyNames;
+			BodyModifierInst->ContactBodyNames = Options.ContactBodyNames;
+			BodyModifierInst->NotifyClass = Options.NotifyClass;
+		}
+		else
+		{
+			UpdateModifierInst->NotifyClass = Options.NotifyClass;
 		}
 	}
 }
 
-void URelativeBodyAnimBlueprintFunctions::BulkUpdatePropBodyModifiers(UPARAM(ref) TArray<UAnimSequenceBase*>& AnimationList, const FPropBodyAnimModifierOptions& Options)
+void URelativeBodyAnimBlueprintFunctions::BulkUpdateRelativeBodyModifiers(TArray<UAnimSequenceBase*>& AnimationList, const FRelativeBodyAnimModifierOptions& Options)
 {
-	// Just run body mod update if the options don't contain prop info
-	if (Options.PropNotifyClass == nullptr || Options.PropMeshSettings.IsEmpty())
-	{
-		BulkUpdateRelativeBodyModifiers(AnimationList, Options);
-		return;
-	}
-	
-	FScopedSlowTask ApplyModifiersTask(static_cast<float>(AnimationList.Num()), LOCTEXT("BulkUpdatePropProgress", "Updating and applying prop/body modifiers to animation list"));
+	FScopedSlowTask ApplyModifiersTask(static_cast<float>(AnimationList.Num()), LOCTEXT("BulkUpdateProgress", "Updating and applying modifier to animation list"));
 	ApplyModifiersTask.MakeDialog();
 	for (UAnimSequenceBase* Anim : AnimationList)
 	{
@@ -116,68 +134,75 @@
 			continue;
 		}
 
-		TObjectPtr<URelativePropsAnimModifier> UpdateModifierInst = FindRelativeBodyModifier<URelativePropsAnimModifier>(AssetUserData);
+		TObjectPtr<URelativeBodyAnimModifier> UpdateModifierInst = FindRelativeBodyModifier<URelativeBodyAnimModifier>(AssetUserData);
 		if (!UpdateModifierInst)
 		{
-			if (!AssetUserData->AddAnimationModifierOfClass(Anim, URelativePropsAnimModifier::StaticClass()))
+			if (!AssetUserData->AddAnimationModifierOfClass(Anim, URelativeBodyAnimModifier::StaticClass()))
 			{
 				UE_LOG(LogAnimation, Warning, TEXT("Unable to create new modifier for anim: %s"), *Anim->GetName());
 				continue;
 			}
 
-			UpdateModifierInst = FindRelativeBodyModifier<URelativePropsAnimModifier>(AssetUserData);
+			UpdateModifierInst = FindRelativeBodyModifier<URelativeBodyAnimModifier>(AssetUserData);
 		}
 
 		if (!UpdateModifierInst)
 		{
 			continue;
 		}
-		
-		UpdateModifierInst->PropsData.Empty(Options.PropMeshSettings.Num());
-		for (const FPropsInfo& PropInfo : Options.PropMeshSettings)
-		{
-			UpdateModifierInst->PropsData.Add(PropInfo);
-		}
-		UpdateModifierInst->bPropsFloorInfoBaking = Options.bPropsFloorInfoBaking;
-		UpdateModifierInst->PropsNotifyClass = Options.PropNotifyClass;
-		
-		// Always set these
+
 		UpdateModifierInst->SampleRate = Options.SampleRate;
 		UpdateModifierInst->ContactThreshold = Options.ContactThreshold;
 		UpdateModifierInst->SkeletalMeshAsset = Options.SkeletalMeshAsset;
 		UpdateModifierInst->PhysicsAssetOverride = Options.PhysicsAssetOverride;
 		UpdateModifierInst->DomainBodyNames = Options.DomainBodyNames;
 		UpdateModifierInst->ContactBodyNames = Options.ContactBodyNames;
-		UpdateModifierInst->NotifyClass = nullptr;
-		
-		// TODO: Ideally we would just delete the old body modifier in this case
-		bool bUpdateSeparateBodyModifier = false;
-		TObjectPtr<URelativeBodyAnimModifier> BodyModifierInst = FindExactTypeBodyModifier<URelativeBodyAnimModifier>(AssetUserData);
-		if (Options.NotifyClass != nullptr && BodyModifierInst != nullptr)
-		{
-			// Update body modifier to generate the 
-			bUpdateSeparateBodyModifier = true;
-			BodyModifierInst->SampleRate = Options.SampleRate;
-			BodyModifierInst->ContactThreshold = Options.ContactThreshold;
-			BodyModifierInst->SkeletalMeshAsset = Options.SkeletalMeshAsset;
-			BodyModifierInst->PhysicsAssetOverride = Options.PhysicsAssetOverride;
-			BodyModifierInst->DomainBodyNames = Options.DomainBodyNames;
-			BodyModifierInst->ContactBodyNames = Options.ContactBodyNames;
-			BodyModifierInst->NotifyClass = Options.NotifyClass;
-		}
-		else
-		{
-			UpdateModifierInst->NotifyClass = Options.NotifyClass;
-		}
+		UpdateModifierInst->NotifyClass = Options.NotifyClass;
 
 		UAnimSequence* AnimSeq = Cast<UAnimSequence>(Anim);
 		if (AnimSeq)
 		{
 			UpdateModifierInst->ApplyToAnimationSequence(AnimSeq);
-			if (bUpdateSeparateBodyModifier)
-			{
-				BodyModifierInst->ApplyToAnimationSequence(AnimSeq);
-			}
 		}
 	}
 }
+
+void URelativeBodyAnimBlueprintFunctions::BulkUpdatePropBodyModifiers(UPARAM(ref) TArray<UAnimSequenceBase*>& AnimationList, const FPropBodyAnimModifierOptions& Options)
+{
+	FScopedSlowTask ApplyModifiersTask(static_cast<float>(AnimationList.Num()), LOCTEXT("BulkUpdatePropProgress", "Updating and applying prop/body modifiers to animation list"));
+	ApplyModifiersTask.MakeDialog();
+	SetRelativeAnimModifiers(AnimationList, Options);
+	
+	for (UAnimSequenceBase* Anim : AnimationList)
+	{
+		ApplyModifiersTask.EnterProgressFrame(1.0f, FText::FromString(Anim->GetName()));
+		UAnimationModifiersAssetUserData* AssetUserData = Anim->GetAssetUserData<UAnimationModifiersAssetUserData>();
+		if (!AssetUserData)
+		{
+			UE_LOG(LogAnimation, Warning, TEXT("Unable to find modifier user data for anim: %s"), *Anim->GetName());
+			continue;
+		}
+
+		TObjectPtr<URelativePropsAnimModifier> PropModifierInst = FindRelativeBodyModifier<URelativePropsAnimModifier>(AssetUserData);
+		TObjectPtr<URelativeBodyAnimModifier> BodyModifierInst = FindExactTypeBodyModifier<URelativeBodyAnimModifier>(AssetUserData);
+		if (!PropModifierInst && !BodyModifierInst)
+		{
+			continue;
+		}
+
+		UAnimSequence* AnimSeq = Cast<UAnimSequence>(Anim);
+		if (!AnimSeq)
+		{
+			continue;
+		}
+		
+		if (PropModifierInst)
+		{
+			PropModifierInst->ApplyToAnimationSequence(AnimSeq);
+		}
+		if (BodyModifierInst)
+		{
+			BodyModifierInst->ApplyToAnimationSequence(AnimSeq);
+		}
+	}
+}
diff --git a/Engine/Plugins/Experimental/Animation/RelativeIKOp/Source/RelativeBodyAnimUtils/Public/RelativeBodyBlueprintFunctions.h b/Engine/Plugins/Experimental/Animation/RelativeIKOp/Source/RelativeBodyAnimUtils/Public/RelativeBodyBlueprintFunctions.h
--- a/Engine/Plugins/Experimental/Animation/RelativeIKOp/Source/RelativeBodyAnimUtils/Public/RelativeBodyBlueprintFunctions.h
+++ b/Engine/Plugins/Experimental/Animation/RelativeIKOp/Source/RelativeBodyAnimUtils/Public/RelativeBodyBlueprintFunctions.h
@@ -74,6 +74,12 @@ public:
 	UFUNCTION(BlueprintCallable, Category = "Anim Modifier")
 	static FName GetPropAttachBone(const FPropsInfo& Info);
 	
+	UFUNCTION(BlueprintCallable, Category = "Anim Notify")
+	static float GetSegmentStartTime(const FAnimNotifyEvent& AnimNotifyEvent);
+	
+	UFUNCTION(BlueprintCallable, Category = "Anim Modifier")
+	static void SetRelativeAnimModifiers(UPARAM(ref) TArray<UAnimSequenceBase*>& AnimationList, const FPropBodyAnimModifierOptions& Options);
+	
 	UFUNCTION(BlueprintCallable, Category = "Anim Modifier")
 	static void BulkUpdateRelativeBodyModifiers(UPARAM(ref) TArray<UAnimSequenceBase*>& AnimationList, const FRelativeBodyAnimModifierOptions& Options);
 	
