RelativeIKOp: Added empty map checks in multibone solve cases.

#Synced-CL 50817256

diff --git a/Engine/Plugins/Experimental/Animation/RelativeIKOp/Source/RelativeIKOp/Private/RelativeIKOp.cpp b/Engine/Plugins/Experimental/Animation/RelativeIKOp/Source/RelativeIKOp/Private/RelativeIKOp.cpp
--- a/Engine/Plugins/Experimental/Animation/RelativeIKOp/Source/RelativeIKOp/Private/RelativeIKOp.cpp
+++ b/Engine/Plugins/Experimental/Animation/RelativeIKOp/Source/RelativeIKOp/Private/RelativeIKOp.cpp
@@ -922,6 +922,12 @@ void FIKRetargetRelativeIKOp::ApplyMatrixBodyPairTargets(
 		OutBoneLocations.Add(TargetBonePose.GetLocation());
 	}
 	
+	// TODO: We need to move all processing to be "output-oriented" (effectors) when we generalize solver!
+	if (OutIndexMap.IsEmpty())
+	{
+		return;
+	}
+	
 	Eigen::MatrixXd BaseWeightMatrix = Eigen::MatrixXd::Identity(OutIndexMap.Num(), OutIndexMap.Num());
 	Eigen::MatrixXd OffsetMatrix = Eigen::MatrixXd::Zero(OutIndexMap.Num(), 3);
 	Eigen::MatrixXd DeltaWeights = Eigen::MatrixXd::Zero(OutIndexMap.Num(), 3);
@@ -1051,6 +1057,11 @@ void FIKRetargetRelativeIKOp::ApplyMatrixBodyPairTargetsWithPropScale(
 	        ScalarBoneNames.Add(SourceBoneName);
 		}
 	}
+	
+	if (OutIndexMap.IsEmpty())
+	{
+		return;
+	}
 
 	Eigen::MatrixXd WeightMatrix = Eigen::MatrixXd::Zero(3*OutIndexMap.Num()+ScalarIndexMap.Num(), 3*OutIndexMap.Num()+ScalarIndexMap.Num());
 	Eigen::MatrixXd OffsetMatrix = Eigen::MatrixXd::Zero(3*OutIndexMap.Num()+ScalarIndexMap.Num(), 1);
