RelativeIKOp: Use new virtual for ik retarget batch operation.

#Synced-CL 50824856

diff --git a/Engine/Plugins/Experimental/Animation/RelativeIKOp/Source/RelativeIKOp/Private/RelativeIKOp.cpp b/Engine/Plugins/Experimental/Animation/RelativeIKOp/Source/RelativeIKOp/Private/RelativeIKOp.cpp
--- a/Engine/Plugins/Experimental/Animation/RelativeIKOp/Source/RelativeIKOp/Private/RelativeIKOp.cpp
+++ b/Engine/Plugins/Experimental/Animation/RelativeIKOp/Source/RelativeIKOp/Private/RelativeIKOp.cpp
@@ -1893,8 +1893,7 @@ void FIKRetargetRelativeIKOp::AnimGraphPreUpdateMainThread(USkeletalMeshComponen
 			PreUpdateMontagePlayhead();
 			return;
 		}
-
-		ResetCacheNotifyInfo();
+		
 		SetupRelativeIKNotifyInfoMontage();
 		return;
 	}
@@ -1981,21 +1980,24 @@ void FIKRetargetRelativeIKOp::PreUpdateMontagePlayhead()
 {
 	if (!MontageInstance)
 	{
-		AnimSeqPlayHead = -1.0f;
+		UpdateAnimPlayhead(-1.0f);
 		return;
 	}
-
-	// TODO: DeltaTimeRecord doesn't seem to work for montages when scrubbing
+	
 	float MontageTime = MontageInstance->GetPosition();
-	AnimSeqPlayHead = MontageTime - SegmentStartTime;
-	if (MontageTime > SegmentEndTime)
+	if (MontageTime < SegmentStartTime || MontageTime > SegmentEndTime)
 	{
-		AnimSeqPlayHead = -1.0f;
+		UpdateAnimPlayhead(-1.0f);
+		return;
 	}
+	
+	UpdateAnimPlayhead(MontageTime - SegmentStartTime);
 }
 
 void FIKRetargetRelativeIKOp::SetupRelativeIKNotifyInfoMontage()
 {
+	ResetCacheNotifyInfo();
+	
 	CacheSourceAnimSequence = nullptr;
 	UAnimMontage* Montage = MontageInstance->Montage;
 	if (!Montage)
@@ -2008,8 +2010,8 @@ void FIKRetargetRelativeIKOp::SetupRelativeIKNotifyInfoMontage()
 		for (const FAnimSegment& Segment : SlotTrack.AnimTrack.AnimSegments)
 		{
 			const TObjectPtr<UAnimSequenceBase>& SeqRef = Segment.GetAnimReference();
-			CacheSourceAnimSequence = Cast<UAnimSequence>(SeqRef.Get());
-			if (!CacheSourceAnimSequence)
+			UAnimSequence* AnimSeq = Cast<UAnimSequence>(SeqRef.Get());
+			if (!AnimSeq)
 			{
 				continue;
 			}
@@ -2017,20 +2019,16 @@ void FIKRetargetRelativeIKOp::SetupRelativeIKNotifyInfoMontage()
 			// TODO: This loop only really works w/ at most one of each notify
 			for (const FAnimNotifyEvent& NotifyEvent : SeqRef->Notifies)
 			{
-				if (URelativeBodyBakeAnimNotify* BodyInfo = Cast<URelativeBodyBakeAnimNotify>(NotifyEvent.Notify))
-				{
-					SegmentStartTime = Segment.AnimStartTime;
-					SegmentEndTime = Segment.AnimEndTime;
-					PreUpdateMontagePlayhead();
-					UpdateCacheNotifyInfo(BodyInfo);
-				}
-				
-				else if (URelativePropsBakeAnimNotify* PropInfo = Cast<URelativePropsBakeAnimNotify>(NotifyEvent.Notify))
+				const TObjectPtr<UAnimNotify>& Notify = NotifyEvent.Notify;
+				if (Cast<URelativeBodyBakeAnimNotify>(Notify) || Cast<URelativePropsBakeAnimNotify>(Notify))
 				{
+					UpdateAnimSeqData(AnimSeq);
+					
+					// Have to cache these for quick updating montage frame-time
 					SegmentStartTime = Segment.AnimStartTime;
 					SegmentEndTime = Segment.AnimEndTime;
 					PreUpdateMontagePlayhead();
-					UpdateCachePropNotifyInfo(PropInfo);
+					return;
 				}
 			}
 		}
@@ -2084,21 +2082,37 @@ void FIKRetargetRelativeIKOp::UpdateRelativeIKNotifyInfoAnimSeq(UAnimInstance* S
 	// Find all the anim sequences (no sync groups) we updated this tick.
 	ProcessTickRecords(UngroupedTickRecords);
 	
-	if (!MaxNotifyRecord)
+	float FrameTime = 0.0f;
+	UAnimSequence* AnimSeq = nullptr;
+	if (MaxNotifyRecord)
 	{
-		CacheSourceAnimSequence = nullptr;
-		ResetCacheNotifyInfo();
-		return;
+		FrameTime = MaxNotifyRecord->DeltaTimeRecord->GetPrevious() + MaxNotifyRecord->DeltaTimeRecord->Delta;
+		AnimSeq = Cast<UAnimSequence>(MaxNotifyRecord->SourceAsset);
 	}
 
-	AnimSeqPlayHead = MaxNotifyRecord->DeltaTimeRecord->GetPrevious() + MaxNotifyRecord->DeltaTimeRecord->Delta;
-	if (MaxNotifyRecord->SourceAsset == CacheSourceAnimSequence)
+	UpdateAnimPlayhead(FrameTime);
+	UpdateAnimSeqData(AnimSeq);
+}
+
+void FIKRetargetRelativeIKOp::UpdateAnimPlayhead(float FrameTime)
+{
+	AnimSeqPlayHead = FrameTime;
+}
+
+void FIKRetargetRelativeIKOp::UpdateAnimSeqData(UAnimSequence* Sequence)
+{
+	if (Sequence == CacheSourceAnimSequence)
 	{
 		return;
 	}
 
 	ResetCacheNotifyInfo();
-	CacheSourceAnimSequence = Cast<UAnimSequence>(MaxNotifyRecord->SourceAsset);
+	CacheSourceAnimSequence = Sequence;
+	
+	if (!Sequence)
+	{
+		return;
+	}
 
 	// TODO: This loop only really works w/ at most one of each notify
 	for (const FAnimNotifyEvent& NotifyEvent : CacheSourceAnimSequence->Notifies)
@@ -2354,6 +2368,12 @@ void FIKRetargetRelativeIKOp::OnParentReinitPropertyEdited(const FIKRetargetOpBa
 
 #if WITH_EDITOR
 
+void FIKRetargetRelativeIKOp::UpdateFromAnimSequence(UAnimSequence* Sequence, float FrameTime)
+{
+	UpdateAnimPlayhead(FrameTime);
+	UpdateAnimSeqData(Sequence);
+}
+
 FCriticalSection FIKRetargetRelativeIKOp::DebugDataMutex;
 
 void FIKRetargetRelativeIKOp::DebugDraw(
diff --git a/Engine/Plugins/Experimental/Animation/RelativeIKOp/Source/RelativeIKOp/Public/RelativeIKOp.h b/Engine/Plugins/Experimental/Animation/RelativeIKOp/Source/RelativeIKOp/Public/RelativeIKOp.h
--- a/Engine/Plugins/Experimental/Animation/RelativeIKOp/Source/RelativeIKOp/Public/RelativeIKOp.h
+++ b/Engine/Plugins/Experimental/Animation/RelativeIKOp/Source/RelativeIKOp/Public/RelativeIKOp.h
@@ -339,6 +339,8 @@ struct FIKRetargetRelativeIKOp : public FIKRetargetOpBase
 	FIKRetargetRelativeIKOpSettings Settings;
 
 private:
+	void UpdateAnimPlayhead(float FrameTime);
+	void UpdateAnimSeqData(UAnimSequence* Sequence);
 	void UpdateCacheBoneChains(const FIKRetargetProcessor& InProcessor, const FIKRetargetRunIKRigOp* ParentRigOp);
 	void UpdateSourceBoneMap(FName SourceBoneName);
 	void UpdateTargetBoneMap(FName TargetBoneName);
@@ -503,6 +505,8 @@ public:
 		const FIKRetargetDebugDrawState& InEditorState) const override;
 
 	virtual bool HasDebugDrawing() override { return true; };
+	
+	UE_API virtual void UpdateFromAnimSequence(UAnimSequence* Sequence, float FrameTime) override;
 
 private:
 	void DebugDrawBodyPairs(FPrimitiveDrawInterface* InPDI, const FTransform& BaseTransform, double Scale, const TArray<FDebugDrawBodyPair>& BodyPairs) const;
