BodyIntersectOp: Initial support for prop intersection and multi-bone intersect motion.

#Synced-CL 50915694

diff --git a/Engine/Plugins/Experimental/Animation/RelativeIKOp/Source/BodyIntersectIKOp/Private/BodyIntersectIKOp.cpp b/Engine/Plugins/Experimental/Animation/RelativeIKOp/Source/BodyIntersectIKOp/Private/BodyIntersectIKOp.cpp
--- a/Engine/Plugins/Experimental/Animation/RelativeIKOp/Source/BodyIntersectIKOp/Private/BodyIntersectIKOp.cpp
+++ b/Engine/Plugins/Experimental/Animation/RelativeIKOp/Source/BodyIntersectIKOp/Private/BodyIntersectIKOp.cpp
@@ -16,7 +16,6 @@
 #include "Engine/EngineTypes.h"
 #endif
 
-
 #include UE_INLINE_GENERATED_CPP_BY_NAME(BodyIntersectIKOp)
 
 #define LOCTEXT_NAMESPACE "BodyIntersectIKOp"
@@ -77,21 +76,6 @@ bool FIKRetargetBodyIntersectIKOp::Initialize(
 	}
 	
 	InitPoleVectorIntersectors(*TargetChains, InTargetSkeleton);
-
-	UpdateCacheSkelInfo(InTargetSkeleton);
-
-	const FIKRigGoalContainer& GoalContainer = InProcessor.GetIKRigGoalContainer();
-	for (int i = 0; i < Settings.IntersectGoalSettings.Num(); i++)
-	{
-		FIKGoalIntersectShapeSettings& GoalSettings = Settings.IntersectGoalSettings[i];
-		const FIKRigGoal* IKGoal = GoalContainer.FindGoalByName(GoalSettings.Goal);
-		if (!IKGoal)
-		{
-			continue;
-		}
-		
-		UpdateTargetBoneMap(IKGoal->BoneName);
-	}
 	
 	bIsInitialized = true;
 	return bIsInitialized;
@@ -117,11 +101,20 @@ void FIKRetargetBodyIntersectIKOp::Run(
 		return;
 	}
 
-#if WITH_EDITOR
 	FDebugBodyIntersectDrawInfo LocalDebugInfo;
+#if WITH_EDITOR
 	LocalDebugInfo.TargetIntersectDetected.SetNumZeroed(Settings.IntersectBodies.Num());
 #endif //WITH_EDITOR
+	
+	// Clear move info every frame (setup in init instead?)
+	MoveWithPropGroups.Init({},Settings.PropIntersectSettings.Num());
 
+	// Store deltas for moving together (max -> blend)
+	FrameDeltaInfo.Init({}, Settings.IntersectGoalSettings.Num());
+	
+	// Clear update prop locations
+	UpdatePropDeltas.Reset(Settings.PropIntersectSettings.Num());
+	
 	FIKRigGoalContainer& GoalContainer = InProcessor.GetIKRigGoalContainer();
 	for (int i = 0; i < Settings.IntersectGoalSettings.Num(); i++)
 	{
@@ -133,73 +126,164 @@ void FIKRetargetBodyIntersectIKOp::Run(
 		}
 
 		const FName GoalBoneName = IKGoal->BoneName;
-		if (!CacheTargetSkelIndices.Contains(GoalBoneName))
+		int32 TargetBoneIdx = InProcessor.GetTargetSkeleton().FindBoneIndexByName(GoalBoneName);
+		if (TargetBoneIdx == INDEX_NONE)
 		{
-			// UE_LOG(LogAnimation, Warning, TEXT("---- Could not find cache bone: %s"), *GoalBoneName.ToString());
 			continue;
 		}
-
-		const int32 TargetBoneIdx = CacheTargetSkelIndices[GoalBoneName];
+		
 		const FTransform TargetBoneTfm = OutTargetGlobalPose[TargetBoneIdx];
+		UpdatePropMoves(GoalSettings.PropIntersectBone, i);
 
+		int32 CheckCollisionIdx = PhysicsAsset->FindBodyIndex(IKGoal->BoneName);
+		if (CheckCollisionIdx == INDEX_NONE)
+		{
+			continue;
+		}
+		
 		FKShapeElem* GoalShape = FPhysShapeUtils::FindBodyShape(PhysicsAsset, IKGoal->BoneName);
 		double GoalRadius = FPhysShapeUtils::CalcShapeSmallRadius(GoalShape); // * GoalSettings.SphereScale;
 		if (FMath::IsNearlyZero(GoalRadius))
 		{
 			continue;
 		}
-
-		// TODO: Handle actual rotation w/ a lever arm and second sphere
+		
 		FVector CompGoalUpdate = GoalLocBlendCompSpace(IKGoal, TargetBoneTfm);
-		FVector CompSphereOffset = TargetBoneTfm.TransformVector(GoalSettings.GoalShapeOffset);
-		FVector SphereOffsetGoal = CompGoalUpdate + CompSphereOffset;
-		FTransform GoalTfm = FTransform(TargetBoneTfm.Rotator(), SphereOffsetGoal, TargetBoneTfm.GetScale3D()*GoalSettings.GoalShapeScale);
+		FVector CompGoalShapeOffset = TargetBoneTfm.TransformVector(GoalSettings.GoalShapeOffset);
+		FVector CompGoalShapeLoc = CompGoalUpdate + CompGoalShapeOffset;
+		FTransform CompShapeTfm = FTransform(TargetBoneTfm.Rotator(), CompGoalShapeLoc, TargetBoneTfm.GetScale3D()*GoalSettings.GoalShapeScale);
+		
+		FVector GoalDelta = GetTotalIntersectDelta(CheckCollisionIdx, GoalShape, CompShapeTfm, OutTargetGlobalPose, InProcessor.GetTargetSkeleton(), LocalDebugInfo);
+		FrameDeltaInfo[i] = {IKGoal,CompShapeTfm,CompGoalUpdate,GoalDelta};
+	}
+	
+	for (int PropIdx = 0; PropIdx < Settings.PropIntersectSettings.Num(); ++PropIdx)
+	{
+		const FIKPropIntersectSettings& PropSettings = Settings.PropIntersectSettings[PropIdx];
+		if (PropSettings.BoneName == NAME_None)
+		{
+			continue;
+		}
+		
+		if (!PropSettings.bPinBone && MoveWithPropGroups[PropIdx].IsEmpty())
+		{
+			continue;
+		}
 		
-		for (int j = 0; j < Settings.IntersectBodies.Num(); ++j)
+		// Get prop intersection-delta
+		FVector PropDelta = FVector::ZeroVector;
+		if (!FMath::IsNearlyZero(PropSettings.CapsuleRadius))
 		{
-			if (!Settings.bEnableGoalIntersect)
+			int32 PropBoneIdx = InProcessor.GetTargetSkeleton().FindBoneIndexByName(PropSettings.BoneName);
+			if (PropBoneIdx == INDEX_NONE)
 			{
-				// TODO: Move this test further out but still allow debug rendering of intersect bodies, etc
 				continue;
 			}
 			
-			FName BodyName = Settings.IntersectBodies[j];
-			if (!CacheTargetSkelIndices.Contains(BodyName))
+			FKSphylElem PropShape = FKSphylElem(PropSettings.CapsuleRadius, PropSettings.CapsuleLength);
+			FTransform PropTfm = FTransform(PropSettings.CapsuleRotation, PropSettings.CapsuleTranslation) * OutTargetGlobalPose[PropBoneIdx];
+			
+			PropDelta = GetTotalIntersectDelta(INDEX_NONE, &PropShape, PropTfm, OutTargetGlobalPose, InProcessor.GetTargetSkeleton(), LocalDebugInfo);
+		}
+		
+		// Nothing else to mess with if no linked goals
+		if (MoveWithPropGroups[PropIdx].IsEmpty())
+		{
+			continue;
+		}
+		
+		// Get consensus delta
+		// TODO: Use delta along avg for all bodies?
+		for (int32 GoalIdx : MoveWithPropGroups[PropIdx])
+		{
+			// const FVector& GoalDelta = FrameDeltaInfo[GoalIdx].IntersectDelta;
+			// if (GoalDelta.SquaredLength() > PropDelta.SquaredLength())
+			// {
+			// 	PropDelta = GoalDelta;
+			// }
+			
+			// TODO: Weight starting deltas for each intersect, instead of at end?
+			FIKRigGoal* IKGoal = FrameDeltaInfo[GoalIdx].GoalPtr;
+			if (!IKGoal)
 			{
 				continue;
 			}
-
-			if (PhysicsAsset->FindBodyIndex(GoalBoneName) == INDEX_NONE || PhysicsAsset->FindBodyIndex(BodyName) == INDEX_NONE
-				|| !PhysicsAsset->IsCollisionEnabled(PhysicsAsset->FindBodyIndex(GoalBoneName), PhysicsAsset->FindBodyIndex(BodyName)))
+			
+			int32 CheckCollisionIdx = PhysicsAsset->FindBodyIndex(IKGoal->BoneName);
+			if (CheckCollisionIdx == INDEX_NONE)
 			{
 				continue;
 			}
-				
-			int32 BodyPoseIdx = CacheTargetSkelIndices[BodyName];
-			FTransform BodyTfm = OutTargetGlobalPose[BodyPoseIdx];
-			FKShapeElem* ShapeElem = FPhysShapeUtils::FindBodyShape(PhysicsAsset, BodyName);
-			if (!ShapeElem)
+			
+			FKShapeElem* GoalShape = FPhysShapeUtils::FindBodyShape(PhysicsAsset, IKGoal->BoneName);
+			double GoalRadius = FPhysShapeUtils::CalcShapeSmallRadius(GoalShape); // * GoalSettings.SphereScale;
+			if (FMath::IsNearlyZero(GoalRadius))
 			{
 				continue;
 			}
-
-			FVector DeltaDir = FVector::ZeroVector;
-			double Dist = FBodyIntersectUtils::CalcIntersectionPairDelta(GoalTfm, GoalShape, BodyTfm, ShapeElem, DeltaDir);
-			if (Dist > 0.0)
+			
+			const FVector StartLoc = FrameDeltaInfo[GoalIdx].StartGoalLoc;
+			FTransform GoalShapeTfm =  FrameDeltaInfo[GoalIdx].StartBodyTfm;
+			GoalShapeTfm.SetLocation(StartLoc + PropDelta);
+			
+			PropDelta += GetTotalIntersectDelta(CheckCollisionIdx, GoalShape, GoalShapeTfm, OutTargetGlobalPose, InProcessor.GetTargetSkeleton(), LocalDebugInfo);
+		}
+		
+		// Propagate max delta back (with blend) to goal offsets
+		for (int32 GoalIdx : MoveWithPropGroups[PropIdx])
+		{
+			double MoveWithAlpha = Settings.IntersectGoalSettings[GoalIdx].BlendMultiMaxOffset;
+			FVector& GoalDelta = FrameDeltaInfo[GoalIdx].IntersectDelta;
+			GoalDelta = FMath::Lerp(GoalDelta, PropDelta, MoveWithAlpha);
+		}
+		
+		// Setup post-solve update for prop bones, if we are pinning
+		if (PropSettings.bPinBone)
+		{
+			int32 PropBoneIdx = InProcessor.GetTargetSkeleton().FindBoneIndexByName(PropSettings.BoneName);
+			if (PropBoneIdx == INDEX_NONE)
 			{
-				SphereOffsetGoal += Dist * DeltaDir;
-				GoalTfm.SetTranslation(SphereOffsetGoal);
+				continue;
+			}
+			
+			UpdatePropDeltas.Add({PropBoneIdx, PropDelta});
+		}
+		
 #if WITH_EDITOR
-				LocalDebugInfo.TargetIntersectDetected[j] = true;
-#endif //WITH_EDITOR
+		if (Settings.bDebugDraw && !FMath::IsNearlyZero(PropSettings.CapsuleRadius))
+		{
+			int32 PropBoneIdx = InProcessor.GetTargetSkeleton().FindBoneIndexByName(PropSettings.BoneName);
+			if (PropBoneIdx == INDEX_NONE)
+			{
+				continue;
 			}
+			
+			FTransform PropTfm = FTransform(PropSettings.CapsuleRotation, PropSettings.CapsuleTranslation) * OutTargetGlobalPose[PropBoneIdx];
+			PropTfm.SetTranslation(PropTfm.GetTranslation() + PropDelta);
+			LocalDebugInfo.PropIntersectTfms.Emplace(PropSettings.CapsuleRadius,PropSettings.CapsuleLength, PropTfm);
+		}
+#endif //WITH_EDITOR
+	}
+	
+	for (int i = 0; i < FrameDeltaInfo.Num(); i++)
+	{
+		FIKRigGoal* IKGoal = FrameDeltaInfo[i].GoalPtr;
+		if (!IKGoal)
+		{
+			continue;
 		}
 		
-		SetGoalPosFromCompSpace(IKGoal, TargetBoneTfm, SphereOffsetGoal-CompSphereOffset);
+		FTransform DebugShapeTfm = FrameDeltaInfo[i].StartBodyTfm;
+		const FVector& GoalBaseLoc = FrameDeltaInfo[i].StartGoalLoc;
+		const FVector& GoalDelta = FrameDeltaInfo[i].IntersectDelta;
+		IKGoal->Position = GoalBaseLoc + GoalDelta;
+		IKGoal->PositionAlpha = 1.0f;
+		// SetGoalPosFromCompSpace(IKGoal, TargetBoneTfm, CompGoalUpdate + GoalDelta);
 #if WITH_EDITOR
 		if (Settings.bDebugDraw)
 		{
-			LocalDebugInfo.GoalIntersectTfms.Emplace(GoalBoneName, GoalTfm);
+			DebugShapeTfm.SetTranslation(GoalBaseLoc + GoalDelta);
+			LocalDebugInfo.GoalIntersectTfms.Emplace(IKGoal->BoneName, DebugShapeTfm);
 		}
 #endif //WITH_EDITOR
 	}
@@ -210,18 +294,19 @@ void FIKRetargetBodyIntersectIKOp::Run(
 		FScopeLock ScopeLock(&DebugDataMutex);
 		
 		DebugDrawInfo.GoalIntersectTfms = LocalDebugInfo.GoalIntersectTfms;
+		DebugDrawInfo.PropIntersectTfms = LocalDebugInfo.PropIntersectTfms;
 
 		// Cache debug info for target intersect bodies
 		DebugDrawInfo.TargetIntersectTfms.Reset(Settings.IntersectBodies.Num());
 		DebugDrawInfo.TargetIntersectDetected = LocalDebugInfo.TargetIntersectDetected;
 		for (FName TargetBoneName : Settings.IntersectBodies)
 		{
-			if (!CacheTargetSkelIndices.Contains(TargetBoneName))
+			int32 TargetPoseIdx = InProcessor.GetTargetSkeleton().FindBoneIndexByName(TargetBoneName);
+			if (TargetPoseIdx == INDEX_NONE)
 			{
 				continue;
 			}
 			
-			int32 TargetPoseIdx = CacheTargetSkelIndices[TargetBoneName];
 			DebugDrawInfo.TargetIntersectTfms.Add({TargetBoneName,OutTargetGlobalPose[TargetPoseIdx]});
 		}
 	}
@@ -244,6 +329,13 @@ void FIKRetargetBodyIntersectIKOp::RunAfterParent(
 		return;
 	}
 	
+	// Apply prop offsets if appropriate
+	for (const TPair<int32,FVector>& PropPinDelta : UpdatePropDeltas)
+	{
+		FTransform& AttachTfm = OutTargetGlobalPose[PropPinDelta.Key];
+		AttachTfm.SetTranslation(AttachTfm.GetTranslation() + PropPinDelta.Value);
+	}
+	
 	if (!Settings.bEnablePoleVectorIntersect)
 	{
 		return;
@@ -588,26 +680,83 @@ void FIKRetargetBodyIntersectIKOp::InitPoleVectorIntersectors(
 	}
 }
 
-void FIKRetargetBodyIntersectIKOp::UpdateTargetBoneMap(FName TargetBoneName)
+FVector FIKRetargetBodyIntersectIKOp::GetTotalIntersectDelta(
+	int32 CollisionProxyIndex,
+	FKShapeElem* IntersectShape,
+	const FTransform& IntersectTfm,
+	const TArray<FTransform>& TargetGlobalPose,
+	const FTargetSkeleton& TargetSkeleton,
+	FDebugBodyIntersectDrawInfo& DebugInfo)
 {
-	int32 TargetIdx = TargetBoneNames.Find(TargetBoneName);
-	if (TargetIdx == INDEX_NONE)
+	FVector IntersectDelta = FVector::ZeroVector;
+	if (!Settings.bEnableGoalIntersect)
 	{
-		return;
+		return IntersectDelta;
 	}
 	
-	CacheTargetSkelIndices.Emplace(TargetBoneName, TargetIdx);
+	FTransform UpdateTfm = IntersectTfm;
+	for (int j = 0; j < Settings.IntersectBodies.Num(); ++j)
+	{
+		FName BodyName = Settings.IntersectBodies[j];
+		int32 BodyPoseIdx = TargetSkeleton.FindBoneIndexByName(BodyName);
+		if (BodyPoseIdx == INDEX_NONE)
+		{
+			continue;
+		}
+		
+		int32 BodyPhysIdx = PhysicsAsset->FindBodyIndex(BodyName);
+		if (BodyPhysIdx == INDEX_NONE)
+		{
+			continue;
+		}
+
+		// Allow CollisionProxyIndex = -1 to skip collision-enabled check (for prop capsules)
+		if (CollisionProxyIndex != INDEX_NONE && !PhysicsAsset->IsCollisionEnabled(CollisionProxyIndex, PhysicsAsset->FindBodyIndex(BodyName)))
+		{
+			continue;
+		}
+		
+		FTransform BodyTfm = TargetGlobalPose[BodyPoseIdx];
+		FKShapeElem* BodyShape = FPhysShapeUtils::FindBodyShape(PhysicsAsset, BodyName);
+		if (!BodyShape)
+		{
+			continue;
+		}
+
+		FVector DeltaDir = FVector::ZeroVector;
+		double Dist = FBodyIntersectUtils::CalcIntersectionPairDelta(UpdateTfm, IntersectShape, BodyTfm, BodyShape, DeltaDir);
+		if (Dist > 0.0)
+		{
+			IntersectDelta += Dist * DeltaDir;
+			UpdateTfm.SetTranslation(IntersectTfm.GetTranslation() + IntersectDelta);
+#if WITH_EDITOR
+			DebugInfo.TargetIntersectDetected[j] = true;
+#endif //WITH_EDITOR
+		}
+	}
+	
+	return IntersectDelta;
 }
 
-void FIKRetargetBodyIntersectIKOp::UpdateCacheSkelInfo(const FTargetSkeleton& InTargetSkeleton)
+void FIKRetargetBodyIntersectIKOp::UpdatePropMoves(FName PropLinkBone, int32 IntersectIndex)
 {
-	TargetBoneNames = InTargetSkeleton.BoneNames;
-	CacheTargetSkelIndices.Reset();
+	if (PropLinkBone == NAME_None)
+	{
+		return;
+	}
+		
+	int32 PropSetIdx = Settings.PropIntersectSettings.IndexOfByPredicate(
+		[PropLinkBone](const FIKPropIntersectSettings& Elem)
+		{
+			return PropLinkBone == Elem.BoneName;
+		});
 	
-	for (FName TargetBodyName : Settings.IntersectBodies)
+	if (PropSetIdx == INDEX_NONE)
 	{
-		UpdateTargetBoneMap(TargetBodyName);
+		return;
 	}
+		
+	MoveWithPropGroups[PropSetIdx].Add(IntersectIndex);
 }
 
 void FIKRetargetBodyIntersectIKOp::OnAddedToStack(const UIKRetargeter* InRetargetAsset, const FIKRetargetOpBase* InParentOp)
@@ -701,6 +850,19 @@ void FIKRetargetBodyIntersectIKOp::DebugDraw(
 		FTransform CompTfm = PoleShapeTfm * InComponentTransform;
 		DebugDrawPhysBody(InPDI, CompTfm, ShapeScale, &SphereElem, FLinearColor::Red);
 	}
+	
+	for (const TTuple<float, float, FTransform>& PropIntersect : DebugDrawInfo.PropIntersectTfms)
+	{
+		FKSphylElem CapsuleElem;
+		CapsuleElem.Radius = PropIntersect.Get<0>();
+		CapsuleElem.Length = PropIntersect.Get<1>();
+		
+		FTransform PropShapeTfm = PropIntersect.Get<2>();
+		FVector ShapeScale = PropShapeTfm.GetScale3D()*InComponentScale;
+		PropShapeTfm.RemoveScaling();
+		FTransform CompTfm = PropShapeTfm * InComponentTransform;
+		DebugDrawPhysBody(InPDI, CompTfm, ShapeScale, &CapsuleElem, FLinearColor(1.0, 0.25, 0.0));
+	}
 }
 
 void FIKRetargetBodyIntersectIKOp::DebugDrawPhysBody(FPrimitiveDrawInterface* InPDI, const FTransform& ParentTransform, const FVector& Scale, const FKShapeElem* ShapeElem, const FLinearColor& Color) const
@@ -750,6 +912,7 @@ void FIKRetargetBodyIntersectIKOp::ResetDebugInfo()
 	DebugDrawInfo.TargetIntersectTfms.Reset();
 	DebugDrawInfo.TargetIntersectDetected.Reset();
 	DebugDrawInfo.GoalIntersectTfms.Reset();
+	DebugDrawInfo.PropIntersectTfms.Reset();
 }
 
 #endif //WITH_EDITOR
diff --git a/Engine/Plugins/Experimental/Animation/RelativeIKOp/Source/BodyIntersectIKOp/Public/BodyIntersectIKOp.h b/Engine/Plugins/Experimental/Animation/RelativeIKOp/Source/BodyIntersectIKOp/Public/BodyIntersectIKOp.h
--- a/Engine/Plugins/Experimental/Animation/RelativeIKOp/Source/BodyIntersectIKOp/Public/BodyIntersectIKOp.h
+++ b/Engine/Plugins/Experimental/Animation/RelativeIKOp/Source/BodyIntersectIKOp/Public/BodyIntersectIKOp.h
@@ -1,6 +1,7 @@
 // Copyright Epic Games, Inc. All Rights Reserved.
 #pragma once
 
+#include "Math/TransformVectorized.h"
 #include "Retargeter/IKRetargetOps.h"
 #include "Retargeter/IKRetargetSettings.h"
 
@@ -23,13 +24,46 @@ class USkeletalMesh;
 struct FDebugBodyIntersectDrawInfo
 {
 	// Intersect info
-	TArray<TTuple<FName,FTransform>> TargetIntersectTfms;
 	TArray<bool> TargetIntersectDetected;
+	TArray<TTuple<FName,FTransform>> TargetIntersectTfms;
 	TArray<TTuple<FName,FTransform>> GoalIntersectTfms;
 	
+	TArray<TTuple<float,float,FTransform>> PropIntersectTfms;
+	
 	TArray<TTuple<float,FTransform>> PoleIntersectTfms;
 };
 
+USTRUCT(BlueprintType, meta = (DisplayName = "Prop Intersect Settings"))
+struct FIKPropIntersectSettings
+{
+	GENERATED_BODY()
+	
+	// The name of the attach bone or IK bone this intersection references
+	UPROPERTY(EditAnywhere, BlueprintReadWrite, Category = "Prop Intersect Settings")
+	FName BoneName = NAME_None;
+	
+	// Optionally also move this bone
+	// TODO: Need to figure out how to interact with blend and relative ik pin bones (RunAfterParent+offset?)
+	UPROPERTY(EditAnywhere, BlueprintReadWrite, Category = "Prop Intersect Settings")
+	bool bPinBone = false;
+	
+	// Capsule radius for prop intersector
+	UPROPERTY(EditAnywhere, BlueprintReadWrite, Category = "Prop Intersect Settings", meta=(ClampMin="0.0"))
+	float CapsuleRadius = 0.0f;
+	
+	// Capsule length for prop intersector
+	UPROPERTY(EditAnywhere, BlueprintReadWrite, Category = "Prop Intersect Settings", meta=(ClampMin="0.0"))
+	float CapsuleLength = 0.0f;
+	
+	// Capsule orientation relative to bone
+	UPROPERTY(EditAnywhere, BlueprintReadWrite, Category = "Prop Intersect Settings")
+	FRotator CapsuleRotation = FRotator::ZeroRotator;
+	
+	// Capsule translation relative to bone
+	UPROPERTY(EditAnywhere, BlueprintReadWrite, Category = "Prop Intersect Settings")
+	FVector CapsuleTranslation = FVector::ZeroVector;
+};
+
 USTRUCT(BlueprintType, meta = (DisplayName = "Goal Intersect Settings"))
 struct FIKGoalIntersectShapeSettings
 {
@@ -44,6 +78,12 @@ struct FIKGoalIntersectShapeSettings
 	// Local offset of shape from goal bone
 	UPROPERTY(EditAnywhere, BlueprintReadWrite, Category = "Goal Intersect Settings")
 	FVector GoalShapeOffset = FVector::ZeroVector;
+	// Prop bone/IK bone to move together with (must be in prop intersect list!)
+	UPROPERTY(EditAnywhere, BlueprintReadWrite, Category = "Goal Intersect Settings")
+	FName PropIntersectBone = NAME_None;
+	// Alpha blend on bone motion relative to prop bone
+	UPROPERTY(EditAnywhere, BlueprintReadWrite, Category = "Goal Intersect Settings", meta=(ClampMin="0.0", ClampMax="1.0"))
+	double BlendMultiMaxOffset = 0.0;
 };
 
 USTRUCT(BlueprintType, meta = (DisplayName = "Pole Vector Settings"))
@@ -84,6 +124,10 @@ struct FIKRetargetBodyIntersectIKOpSettings : public FIKRetargetOpSettingsBase
 	// IK Goals Shape intersect against physics asset bodies 
 	UPROPERTY(EditAnywhere, BlueprintReadWrite, Category = "Intersection", meta=(ReinitializeOnEdit))
 	TArray<FIKGoalIntersectShapeSettings> IntersectGoalSettings;
+	
+	// Prop intersection shape and bone info (must match goal intersect names if moving together) 
+	UPROPERTY(EditAnywhere, BlueprintReadWrite, Category = "Intersection", meta=(ReinitializeOnEdit))
+	TArray<FIKPropIntersectSettings> PropIntersectSettings;
 
 	// Bone names with attached physics bodies to do trivial intersection against
 	// **NOTE: Only Sphere, Capsule and Box bodies currently supported**
@@ -178,10 +222,25 @@ struct FIKRetargetBodyIntersectIKOp : public FIKRetargetOpBase
 	FIKRetargetBodyIntersectIKOpSettings Settings;
 
 private:
+	struct FIntersectFrameDeltaInfo
+	{
+		FIKRigGoal* GoalPtr = nullptr;
+		FTransform StartBodyTfm = FTransform::Identity;
+		FVector StartGoalLoc = FVector::ZeroVector;
+		FVector IntersectDelta = FVector::ZeroVector;
+	};
+	
 	void InitPoleVectorIntersectors(const TArray<FResolvedBoneChain>& TargetChains, const FTargetSkeleton& InTargetSkeleton);
 	
-	void UpdateCacheSkelInfo(const FTargetSkeleton& InTargetSkeleton);
-	void UpdateTargetBoneMap(FName TargetBoneName);
+	FVector GetTotalIntersectDelta(
+		int32 CollisionProxyIndex,
+		FKShapeElem* IntersectShape,
+		const FTransform& IntersectTfm,
+		const TArray<FTransform>& TargetGlobalPose,
+		const FTargetSkeleton& TargetSkeleton,
+		FDebugBodyIntersectDrawInfo& DebugInfo);
+	
+	void UpdatePropMoves(FName PropLinkBone, int32 IntersectIndex);
 
 	FVector GoalLocBlendCompSpace(const FIKRigGoal* Goal, const FTransform& BoneTfm) const;
 	void SetGoalPosFromCompSpace(FIKRigGoal* Goal, const FTransform& BoneTfm, const FVector& CompSpaceLoc) const;
@@ -201,6 +260,8 @@ private:
 	void DebugDrawPhysBody(FPrimitiveDrawInterface* InPDI, const FTransform& ParentTransform, const FVector& Scale, const FKShapeElem* ShapeElem, const FLinearColor& Color) const;
 	
 	void ResetDebugInfo();
+	
+	FKShapeElem* CopyShape(FKShapeElem* InShape);
 
 	// Collection of all debug structures
 	FDebugBodyIntersectDrawInfo DebugDrawInfo;
@@ -216,11 +277,15 @@ private:
 	// Array (per-chain w/ pv bodies) of intersectors
 	TArray<FPoleVectorChainIntersector> PoleVectorChainIntersectors;
 	
-	// Cache source/target bone-names
-	TArray<FName> TargetBoneNames;
-
-	// Cache target skel index info
-	TMap<FName,int32> CacheTargetSkelIndices;
+	// Attach bone updates to pin
+	// TODO: Decide if these should be deltas or absolute locs
+	TArray<TPair<int32,FVector>> UpdatePropDeltas;
+	
+	// Per-frame goal deltas
+	TArray<FIntersectFrameDeltaInfo> FrameDeltaInfo;
+	
+	// Move with prop groups
+	TArray<TArray<int32>> MoveWithPropGroups;
 };
 
 /* The blueprint/python API for editing BodyIntersectIK Op */
