RelativePropBake: Move prop attach bake to FName and allow socket or bone names.

#Synced-CL 50924028

diff --git a/Engine/Plugins/Experimental/Animation/RelativeIKOp/Source/RelativeBodyAnimUtils/Private/RelativeBodyBlueprintFunctions.cpp b/Engine/Plugins/Experimental/Animation/RelativeIKOp/Source/RelativeBodyAnimUtils/Private/RelativeBodyBlueprintFunctions.cpp
--- a/Engine/Plugins/Experimental/Animation/RelativeIKOp/Source/RelativeBodyAnimUtils/Private/RelativeBodyBlueprintFunctions.cpp
+++ b/Engine/Plugins/Experimental/Animation/RelativeIKOp/Source/RelativeBodyAnimUtils/Private/RelativeBodyBlueprintFunctions.cpp
@@ -40,12 +40,12 @@ TRelativeBodyClass* FindExactTypeBodyModifier(UAnimationModifiersAssetUserData*
 
 void URelativeBodyAnimBlueprintFunctions::SetPropAttachBone(UPARAM(ref) FPropsInfo& Info, FName AttachBoneName)
 {
-	Info.AttachBone = AttachBoneName;
+	Info.SocketName = AttachBoneName;
 }
 
 FName URelativeBodyAnimBlueprintFunctions::GetPropAttachBone(const FPropsInfo& Info)
 {
-	return Info.AttachBone.BoneName;
+	return Info.SocketName;
 }
 
 float URelativeBodyAnimBlueprintFunctions::GetSegmentStartTime(const FAnimNotifyEvent& AnimNotifyEvent)
diff --git a/Engine/Plugins/Experimental/Animation/RelativeIKOp/Source/RelativeBodyAnimUtils/Private/RelativePropsAnimModifier.cpp b/Engine/Plugins/Experimental/Animation/RelativeIKOp/Source/RelativeBodyAnimUtils/Private/RelativePropsAnimModifier.cpp
--- a/Engine/Plugins/Experimental/Animation/RelativeIKOp/Source/RelativeBodyAnimUtils/Private/RelativePropsAnimModifier.cpp
+++ b/Engine/Plugins/Experimental/Animation/RelativeIKOp/Source/RelativeBodyAnimUtils/Private/RelativePropsAnimModifier.cpp
@@ -23,6 +23,7 @@
 #include "MeshAdapter.h"
 #include "Spatial/MeshAABBTree3.h"
 #include "Algo/MinElement.h"
+#include "Engine/SkeletalMeshSocket.h"
 
 #define LOCTEXT_NAMESPACE "RelativePropsAnimModifier"
 
@@ -276,9 +277,9 @@ void URelativePropsAnimModifier::OnApply_Implementation(UAnimSequence* InAnimati
 	const int32 NumBodies = BodyIndicesToBoneIndices.Num();
 	const float HashGridCellSize = std::min(1.5f * ContactThreshold, 50.f); // the length of the cell size in the point hash gri
 
-	// Check for valid prop data
-	TArray<bool> IsValidPropData;
-	IsValidPropData.SetNumZeroed(PropsData.Num());
+	// Setup and check for valid prop data
+	PropAttachBones.Init(NAME_None, PropsData.Num());
+	PropAttachTfms.Init(FTransform::Identity, PropsData.Num());
 	bool PropsDataRequired = false;
 	for (int32 PropsIdx = 0; PropsIdx < PropsData.Num(); PropsIdx++)
 	{
@@ -295,13 +296,15 @@ void URelativePropsAnimModifier::OnApply_Implementation(UAnimSequence* InAnimati
 			continue;
 		}
 		
-		if (RefSkeleton.FindBoneIndex(PropsData[PropsIdx].AttachBone.BoneName) == INDEX_NONE)
+		FName CheckPropBone = GetPropAttachInfo(PropAttachTfms[PropsIdx], PropsData[PropsIdx], SkeletalMeshAsset);
+		if (CheckPropBone == NAME_None)
 		{
-			UE_LOG(LogAnimation, Warning, TEXT("RelativePropAnimModifier: Bad propdata at index [%d]. \n\tReason: Attach bone '%s' not found in skeletal mesh '%s'!"), PropsIdx, *PropsData[PropsIdx].AttachBone.BoneName.ToString(), *GetNameSafe(SkeletalMeshAsset));
+			UE_LOG(LogAnimation, Warning, TEXT("RelativePropAnimModifier: Bad propdata at index [%d]. \n\tReason: Attach socket/bone '%s' not found in skeletal mesh '%s'!"), PropsIdx, *PropsData[PropsIdx].SocketName.ToString(), *GetNameSafe(SkeletalMeshAsset));
 			continue;
 		}
-
-		IsValidPropData[PropsIdx] = true;
+		
+		// NOTE: PropAttachBones[i] != NAME_None indicates valid prop data
+		PropAttachBones[PropsIdx] = CheckPropBone;
 		PropsDataRequired = true;
 	}
 	
@@ -373,11 +376,11 @@ void URelativePropsAnimModifier::OnApply_Implementation(UAnimSequence* InAnimati
 				
 			for (size_t PropsIdx = 0; PropsIdx < PropsData.Num(); PropsIdx++)
 			{
-				if (!IsValidPropData[PropsIdx])
+				if (PropAttachBones[PropsIdx] == NAME_None)
 				{
 					continue;
 				}
-				BakedAnimNotifyInfo->OffsetTransformsForBones.Emplace(PropsData[PropsIdx].AttachBone.BoneName,  PropsData[PropsIdx].AttachTransform);
+				BakedAnimNotifyInfo->OffsetTransformsForBones.Emplace(PropAttachBones[PropsIdx],  PropAttachTfms[PropsIdx]);
 				FMeshDescription* PropsMeshDescription = PropsData[PropsIdx].GetMeshDescription(0); // LOD 0
 				PropsTreeData[PropsIdx].Init(PropsMeshDescription);
 	
@@ -388,13 +391,13 @@ void URelativePropsAnimModifier::OnApply_Implementation(UAnimSequence* InAnimati
 						continue;
 					}
 					FName BoneName = RefSkeleton.GetBoneName(BodyIndicesToBoneIndices[BodyIndex1]);
-					BakedAnimNotifyInfo->PropsPairs.Add(PropsData[PropsIdx].AttachBone.BoneName);
+					BakedAnimNotifyInfo->PropsPairs.Add(PropAttachBones[PropsIdx]);
 					BakedAnimNotifyInfo->PropsPairs.Add(BoneName);
 					BakedAnimNotifyInfo->bPropsPairsIsParentDominates.Add(IsDomainBodyIndices[BodyIndex1]);
 				}
 				if (bPropsFloorInfoBaking)
 				{
-					BakedAnimNotifyInfo->PropsPairs.Add(PropsData[PropsIdx].AttachBone.BoneName);
+					BakedAnimNotifyInfo->PropsPairs.Add(PropAttachBones[PropsIdx]);
 					BakedAnimNotifyInfo->PropsPairs.Add(FName("Floor"));
 					BakedAnimNotifyInfo->bPropsPairsIsParentDominates.Add(true);
 				}
@@ -444,13 +447,13 @@ void URelativePropsAnimModifier::OnApply_Implementation(UAnimSequence* InAnimati
 				BakedAnimNotifyInfo->PropsPairsSampleTime[SampleIndex] = SampleTime;
 				for (int32 PropsIdx = 0; PropsIdx < PropsData.Num(); PropsIdx++)
 				{
-					if (!IsValidPropData[PropsIdx])
+					if (PropAttachBones[PropsIdx] == NAME_None)
 					{
 						continue;
 					}
 					
-					FTransform AttachBoneTransform = UAnimPoseExtensions::GetBonePose(AnimPose, PropsData[PropsIdx].AttachBone.BoneName, EAnimPoseSpaces::World);
-					FTransform PropsLocalToWorld = PropsData[PropsIdx].AttachTransform*AttachBoneTransform;
+					FTransform AttachBoneTransform = UAnimPoseExtensions::GetBonePose(AnimPose, PropAttachBones[PropsIdx], EAnimPoseSpaces::World);
+					FTransform PropsLocalToWorld = PropAttachTfms[PropsIdx] * AttachBoneTransform;
 					const FPositionVertexBuffer& PositionVertexBuffer = PropsData[PropsIdx].GetPositionVertexBuffer();
 					
 					if (PropsData[PropsIdx].PropSkeletalMeshAsset && PropsData[PropsIdx].PropAnimSequence)
@@ -463,7 +466,7 @@ void URelativePropsAnimModifier::OnApply_Implementation(UAnimSequence* InAnimati
 						PropsTreeData[PropsIdx].UpdateVertexPositions(PropsVLocations[PropsIdx]);
 						for (int32 i = 0; i < PropsVLocations[PropsIdx].Num(); ++i)
 						{
-							FVector LocalPos = PropsData[PropsIdx].AttachTransform.TransformPosition(FVector(PropsVLocations[PropsIdx][i]));
+							FVector LocalPos = PropAttachTfms[PropsIdx].TransformPosition(FVector(PropsVLocations[PropsIdx][i]));
 							PropsVLocations[PropsIdx][i] = FVector3f(AttachBoneTransform.TransformPosition(LocalPos));
 						}
 					}
@@ -498,7 +501,7 @@ void URelativePropsAnimModifier::OnApply_Implementation(UAnimSequence* InAnimati
 						// UE_LOG(LogAnimation, Warning, TEXT("MinDIst %f: (%f,%f,%f)."), MinDist, NearestPoint.X, NearestPoint.Y, NearestPoint.Z);
 						if (NearestQueryVertexIndex != INDEX_NONE)
 						{
-							FName BoneName1 = PropsData[PropsIdx].AttachBone.BoneName;
+							FName BoneName1 = PropAttachBones[PropsIdx];
 							FName BoneName2 = RefSkeleton.GetBoneName(BodyIndicesToBoneIndices[BodyIndex]);
 							FVector3f ContactPoint1(PropsLocalToWorld.TransformPosition(NearestPoint));
 							FVector3f ContactPoint2 = VLocations[VertexIndicesInfluencedByBodyIndices[BodyIndex][NearestQueryVertexIndex]];
@@ -518,7 +521,8 @@ void URelativePropsAnimModifier::OnApply_Implementation(UAnimSequence* InAnimati
 					}
 					if (bPropsFloorInfoBaking)
 					{
-						FTransform3f PropOffsetTfm(BakedAnimNotifyInfo->OffsetTransformsForBones[PropsData[PropsIdx].AttachBone.BoneName]);
+						FName PropBone = PropAttachBones[PropsIdx];
+						FTransform3f PropOffsetTfm(BakedAnimNotifyInfo->OffsetTransformsForBones[PropBone]);
 						int32 Offset = SampleIndex * BakedAnimNotifyInfo->PropsPairs.Num() + Count * 2;
 						int32 LowestVertexIndex = FindLowestIndex(PropsVLocations[PropsIdx]);
 						int32 HighestVertexIndex=FindHighestIndex(PropsVLocations[PropsIdx]);
@@ -534,23 +538,47 @@ void URelativePropsAnimModifier::OnApply_Implementation(UAnimSequence* InAnimati
 	}
 }
 
+FName URelativePropsAnimModifier::GetPropAttachInfo(FTransform& OutAttachTfm, const FPropsInfo& PropInfo, const USkeletalMesh* SkeletalMesh)
+{
+	FName AttachBoneName = NAME_None;
+	OutAttachTfm = FTransform::Identity;
+
+	FTransform BaseAttachTfm = FTransform::Identity;
+	const USkeletalMeshSocket* Socket = SkeletalMesh->FindSocket(PropInfo.SocketName);
+	if (Socket)
+	{
+		AttachBoneName = Socket->BoneName;
+		OutAttachTfm = PropInfo.AttachTransform * FTransform(Socket->RelativeRotation, Socket->RelativeLocation, Socket->RelativeScale);
+	}
+	else if (SkeletalMesh->GetRefSkeleton().FindBoneIndex(PropInfo.SocketName) != INDEX_NONE)
+	{
+		AttachBoneName = PropInfo.SocketName;
+		OutAttachTfm = PropInfo.AttachTransform;
+	}
+	
+	return AttachBoneName;
+}
+
 void URelativePropsAnimModifier::GetPropsVertices(TArray<TArray<FVector3f>>& PropsVLocations, const FAnimPose& AnimPose)
 {
-	PropsVLocations.Empty();
+	PropsVLocations.Init({}, PropsData.Num());
 	for (int32 PropsIdx = 0; PropsIdx < PropsData.Num(); PropsIdx++)
 	{
-		const FTransform AttachBoneTransform = UAnimPoseExtensions::GetBonePose(AnimPose, PropsData[PropsIdx].AttachBone.BoneName, EAnimPoseSpaces::World);
+		if (PropAttachBones[PropsIdx] == NAME_None)
+		{
+			continue;
+		}
+		
+		const FTransform AttachBoneTransform = UAnimPoseExtensions::GetBonePose(AnimPose, PropAttachBones[PropsIdx], EAnimPoseSpaces::World);
 		
 		const FPositionVertexBuffer& PositionVertexBuffer = PropsData[PropsIdx].GetPositionVertexBuffer();
-
-		TArray<FVector3f> LocalVertices;
+		
 		int32 NumVerts = PositionVertexBuffer.GetNumVertices();
 		for (int32 i = 0; i < NumVerts; ++i)
 		{
-			FVector LocalPos = PropsData[PropsIdx].AttachTransform.TransformPosition(FVector(PositionVertexBuffer.VertexPosition(i)));
-			LocalVertices.Add(FVector3f(AttachBoneTransform.TransformPosition(LocalPos)));
+			FVector LocalPos = PropAttachTfms[PropsIdx].TransformPosition(FVector(PositionVertexBuffer.VertexPosition(i)));
+			PropsVLocations[PropsIdx].Add(FVector3f(AttachBoneTransform.TransformPosition(LocalPos)));
 		}
-		PropsVLocations.Add(LocalVertices);
 	}
 
 }
diff --git a/Engine/Plugins/Experimental/Animation/RelativeIKOp/Source/RelativeBodyAnimUtils/Public/RelativeBodyBlueprintFunctions.h b/Engine/Plugins/Experimental/Animation/RelativeIKOp/Source/RelativeBodyAnimUtils/Public/RelativeBodyBlueprintFunctions.h
--- a/Engine/Plugins/Experimental/Animation/RelativeIKOp/Source/RelativeBodyAnimUtils/Public/RelativeBodyBlueprintFunctions.h
+++ b/Engine/Plugins/Experimental/Animation/RelativeIKOp/Source/RelativeBodyAnimUtils/Public/RelativeBodyBlueprintFunctions.h
@@ -68,10 +68,10 @@ class RELATIVEBODYANIMUTILS_API URelativeBodyAnimBlueprintFunctions : public UBl
 	GENERATED_BODY()
 public:
 	
-	UFUNCTION(BlueprintCallable, Category = "Anim Modifier")
+	UFUNCTION(BlueprintCallable, Category = "Anim Modifier", meta=(DeprecatedFunction, DeprecationMessage="Set prop attach SocketName directly in Props Info"))
 	static void SetPropAttachBone(UPARAM(ref) FPropsInfo& Info, FName AttachBoneName);
 	
-	UFUNCTION(BlueprintCallable, Category = "Anim Modifier")
+	UFUNCTION(BlueprintCallable, Category = "Anim Modifier", meta=(DeprecatedFunction, DeprecationMessage="Get prop attach SocketName directly from Props Info"))
 	static FName GetPropAttachBone(const FPropsInfo& Info);
 	
 	UFUNCTION(BlueprintCallable, Category = "Anim Notify")
diff --git a/Engine/Plugins/Experimental/Animation/RelativeIKOp/Source/RelativeBodyAnimUtils/Public/RelativePropsAnimModifier.h b/Engine/Plugins/Experimental/Animation/RelativeIKOp/Source/RelativeBodyAnimUtils/Public/RelativePropsAnimModifier.h
--- a/Engine/Plugins/Experimental/Animation/RelativeIKOp/Source/RelativeBodyAnimUtils/Public/RelativePropsAnimModifier.h
+++ b/Engine/Plugins/Experimental/Animation/RelativeIKOp/Source/RelativeBodyAnimUtils/Public/RelativePropsAnimModifier.h
@@ -30,8 +30,8 @@ struct FPropsInfo
 	TObjectPtr<USkeletalMesh> PropSkeletalMeshAsset;
 
 	// Source skeletal mesh bone to attach prop to
-	UPROPERTY(EditAnywhere, Category = "Prop", meta=(ReinitializeOnEdit))
-	FBoneReference AttachBone;
+	UPROPERTY(EditAnywhere, BlueprintReadWrite, Category = "Prop", meta=(ReinitializeOnEdit))
+	FName SocketName;
 
 	// Local transform of prop relative to attach bone
 	UPROPERTY(EditAnywhere, BlueprintReadWrite, Category = "Prop", meta=(ReinitializeOnEdit))
@@ -81,6 +81,15 @@ public:
 
 	UPROPERTY(EditAnywhere, Category = "Settings", meta = (TitleProperty = "PropsFloorInfoBaking"))
 	bool bPropsFloorInfoBaking = false;
-protected:
+	
+private:
 	void GetPropsVertices(TArray<TArray<FVector3f>>& PropsVLocations, const FAnimPose& AnimPose);
+	
+	static FName GetPropAttachInfo(FTransform& OutAttachTfm, const FPropsInfo& PropInfo, const USkeletalMesh* SkeletalMesh);
+	
+	// Actual bone names (pulled from socket info or prop info)
+	TArray<FName> PropAttachBones;
+	
+	// Baked down socket/prop attachments
+	TArray<FTransform> PropAttachTfms;
 };
\ No newline at end of file
