RelativeIKOp: Amplify "important" IKGoal weights against prop attaches when distances are very close

#Synced-CL 51540788

diff --git a/Engine/Plugins/Experimental/Animation/RelativeIKOp/Source/RelativeIKOp/Private/RelativeIKOp.cpp b/Engine/Plugins/Experimental/Animation/RelativeIKOp/Source/RelativeIKOp/Private/RelativeIKOp.cpp
--- a/Engine/Plugins/Experimental/Animation/RelativeIKOp/Source/RelativeIKOp/Private/RelativeIKOp.cpp
+++ b/Engine/Plugins/Experimental/Animation/RelativeIKOp/Source/RelativeIKOp/Private/RelativeIKOp.cpp
@@ -157,9 +157,23 @@ void FIKRetargetRelativeIKOp::Run(
 	
 	LocalTargetPairSpaces.SetNumZeroed(TotalBodies);
 #endif //WITH_EDITOR
+	
+	FIKRigGoalContainer& GoalContainer = InProcessor.GetIKRigGoalContainer();
+	TArray<FIKRigGoal>& IKGoals = GoalContainer.GetGoalArray();
 
 	// Copy/cache (pin) prop bones in target using pin bone settings
 	CachePinBoneTransforms(InProcessor, InSourceGlobalPose, OutTargetGlobalPose);
+	
+	// Get Array of bones we may amplify weight on
+	TArray<FName> AmplifyBones;
+	if (Settings.bApplyImportanceWeighting)
+	{
+		AmplifyBones.Reset(IKGoals.Num() + Settings.PinPropBones.Num());
+		for (const FIKRigGoal& Goal : IKGoals)
+		{
+			AmplifyBones.Add(Goal.BoneName);
+		}
+	}
 
 	// TODO: Should we just use initial chain lengths?
 	CacheChainLengthMap.Reset();
@@ -178,11 +192,13 @@ void FIKRetargetRelativeIKOp::Run(
 	{
 		int32 NumPropSamples = CachedPropNotifyInfo->NumSamples;
 		TConstArrayView<FVector3f> FramePropView = ApplyTemporalSmoothing(AnimSeqPlayHead, PropSampleRate, NumPropSamples, NumPropBones, CachedPropNotifyInfo->PropsPairsLocalReference);
-		ComputeFramePropSpaces(SourceScale, InSourceGlobalPose, OutTargetGlobalPose, CachePropPairOffset, FramePropView, LocalDebugInfo, LocalTargetPairSpaces);
+		ComputeFramePropSpaces(SourceScale, AmplifyBones, InSourceGlobalPose, OutTargetGlobalPose, CachePropPairOffset, FramePropView, LocalDebugInfo, LocalTargetPairSpaces);
 	}
 	
-	FIKRigGoalContainer& GoalContainer = InProcessor.GetIKRigGoalContainer();
-	TArray<FIKRigGoal>& IKGoals = GoalContainer.GetGoalArray();
+	if (Settings.bApplyImportanceWeighting)
+	{
+		AmplifyImportanceWeights();
+	}
 	
 	LocalDebugInfo.TargetGoals.Reserve(CacheSourceBodyEffectVertIdx.Num());
 	
@@ -399,6 +415,8 @@ void FIKRetargetRelativeIKOp::ComputeFramePairSpaces(
 		FrameBoneConstraints[2*OutPairIdx+1].FeasibilityWeight = FeasibilityWeight_O;
 		FrameBoneConstraints[2*OutPairIdx].Weight = Weight;
 		FrameBoneConstraints[2*OutPairIdx+1].Weight = Weight;
+		FrameBoneConstraints[2*OutPairIdx].SourceDist = SourceVDist;
+		FrameBoneConstraints[2*OutPairIdx+1].SourceDist = SourceVDist;
 		
 #if WITH_EDITOR
 		// Debug draw verts
@@ -423,6 +441,7 @@ void FIKRetargetRelativeIKOp::ComputeFramePairSpaces(
 
 void FIKRetargetRelativeIKOp::ComputeFramePropSpaces(
 	double SourceScale,
+	const TArray<FName>& AmplifyBones,
 	const TArray<FTransform>& SourcePose,
 	const TArray<FTransform>& TargetPose,
 	int32 PairOffset,
@@ -492,7 +511,7 @@ void FIKRetargetRelativeIKOp::ComputeFramePropSpaces(
 			double TargetZDist = TargetPropCompZMin.Z;
 			double TargetWeight = GetDistanceWeight(Settings.DistanceWeightMode, TargetZDist, Settings.FloorThreshold, Settings.DistHalfPointLambda);
 			
-			// Offset vector (Prop -> Floor) to Lerp for dist alph
+			// Offset vector (Prop -> Floor) to Lerp for dist alpha
 			double FloorOffsetZMin = Settings.DistHalfPointLambda*Settings.FloorThreshold-TargetPropCompZMin.Z;
 			FVector TargetOffsetZMin = SourcePropCompZMin - TargetPropCompZMin;
 			FVector TargetOffsetZMax = SourcePropCompZMax - TargetPropCompZMax;
@@ -506,6 +525,8 @@ void FIKRetargetRelativeIKOp::ComputeFramePropSpaces(
 			FrameBoneConstraints[2*OutPairIdx+1].Offset = {0,0,SourceWeight*TargetOffsetZMax.Z + (1-SourceWeight)*FloorOffsetZMin};
 			FrameBoneConstraints[2*OutPairIdx].Weight = Weight;
 			FrameBoneConstraints[2*OutPairIdx+1].Weight = Weight;
+			FrameBoneConstraints[2*OutPairIdx].SourceDist = SourceZDist;
+			FrameBoneConstraints[2*OutPairIdx+1].SourceDist = SourceZDist;
 
 #if WITH_EDITOR
 			// Debug draw verts
@@ -533,7 +554,8 @@ void FIKRetargetRelativeIKOp::ComputeFramePropSpaces(
 			continue;
 		}
 		
-		
+		// Potentially amplify weight of this pair by total relative weight
+		bool bAmplify = AmplifyBones.Contains(TargetBodyBone);
 		
 		int32 SourcePropPoseIdx = CacheSourceSkelIndices[SourcePropBone];
 		int32 SourceBodyPoseIdx = CacheSourceSkelIndices[SourceBodyBone];
@@ -610,7 +632,7 @@ void FIKRetargetRelativeIKOp::ComputeFramePropSpaces(
 		
 		FVector WeightedTargetOffset_B = TargetOffsetVec * (Settings.bMultiBoneSolve? 0.5*(FeasibilityWeight_B+FeasibilityWeight_P): FeasibilityWeight_B);
 		FVector WeightedTargetOffset_P = -TargetOffsetVec * (Settings.bMultiBoneSolve? 0.5*(FeasibilityWeight_B+FeasibilityWeight_P): FeasibilityWeight_P);
-
+		
 		FrameBoneConstraints[2*OutPairIdx].ContactOffset = TargetSpring_P - TargetPropBoneLoc;
 		FrameBoneConstraints[2*OutPairIdx+1].ContactOffset = TargetSpring_B - TargetBodyBoneLoc;
 		FrameBoneConstraints[2*OutPairIdx].Offset = WeightedTargetOffset_P;
@@ -619,6 +641,10 @@ void FIKRetargetRelativeIKOp::ComputeFramePropSpaces(
 		FrameBoneConstraints[2*OutPairIdx+1].FeasibilityWeight = FeasibilityWeight_B;
 		FrameBoneConstraints[2*OutPairIdx].Weight = Weight;
 		FrameBoneConstraints[2*OutPairIdx+1].Weight = Weight;
+		FrameBoneConstraints[2*OutPairIdx].SourceDist = SourceVDist;
+		FrameBoneConstraints[2*OutPairIdx+1].SourceDist = SourceVDist;
+		FrameBoneConstraints[2*OutPairIdx].bAmplifyCloseWeight = bAmplify;
+		FrameBoneConstraints[2*OutPairIdx+1].bAmplifyCloseWeight = bAmplify;
 
 #if WITH_EDITOR
 		// Debug draw verts
@@ -641,6 +667,82 @@ void FIKRetargetRelativeIKOp::ComputeFramePropSpaces(
 	}
 }
 
+void FIKRetargetRelativeIKOp::AmplifyImportanceWeights()
+{
+	TArray<int32> AmplifyEffectIdx;
+	double MaxEffectWeight = 0.0;
+	AmplifyEffectIdx.Reset(CacheSourceBodyEffectVertIdx.Num());
+	for (int32 EffectIdx = 0; EffectIdx < CacheSourceBodyEffectVertIdx.Num(); ++EffectIdx)
+	{
+		TArray<int32>& EffectIndices = CacheSourceBodyEffectVertIdx[EffectIdx];
+		if (EffectIndices.IsEmpty())
+		{
+			continue;
+		}
+
+		FName SourceBoneName = CacheSourceEffectBones[EffectIdx];
+		FName TargetBoneName = ApplyBodyMap(SourceBoneName);
+		if (!CacheTargetSkelIndices.Contains(TargetBoneName))
+		{
+			continue;
+		}
+		
+		bool bAmplifyBone = false;
+		double TotalWeight = 0.0;
+		for (int Index : EffectIndices)
+		{
+			const FRelativeIKFrameBoneConstraint& BoneConstraint = FrameBoneConstraints[Index];
+			if (FMath::IsNearlyZero(BoneConstraint.Weight))
+			{
+				continue;
+			}
+			
+			if (BoneConstraint.bAmplifyCloseWeight && BoneConstraint.SourceDist <= Settings.AmplifyCloseThreshold)
+			{
+				bAmplifyBone = true;
+				continue;
+			}
+			
+			double TargetWeight = BoneConstraint.Weight;
+			TotalWeight += TargetWeight;
+		}
+		
+		if (bAmplifyBone)
+		{
+			AmplifyEffectIdx.Add(EffectIdx);
+		}
+		
+		// TODO: Is applying max across all amplify weights reasonable?
+		MaxEffectWeight = FMath::Max(MaxEffectWeight, TotalWeight);
+	}
+	
+	const double AmpWeight = MaxEffectWeight * Settings.AmplifyScale;
+	for (int32 EffectIdx : AmplifyEffectIdx)
+	{
+		TArray<int32>& EffectIndices = CacheSourceBodyEffectVertIdx[EffectIdx];
+		
+		for (int Index : EffectIndices)
+		{
+			FRelativeIKFrameBoneConstraint& BoneConstraint = FrameBoneConstraints[Index];
+			if (!BoneConstraint.bAmplifyCloseWeight || BoneConstraint.SourceDist > Settings.AmplifyCloseThreshold)
+			{
+				continue;
+			}
+			
+			const double Weight = BoneConstraint.Weight;
+			if (BoneConstraint.Weight >= AmpWeight)
+			{
+				continue;
+			}
+			
+			double BlendAlpha = FMath::Clamp((Settings.AmplifyCloseThreshold - BoneConstraint.SourceDist) / (Settings.AmplifyCloseThreshold - Settings.AmplifyCloseDistance), 0.0, 1.0);
+			double BlendWeight = FMath::Lerp(BoneConstraint.Weight, AmpWeight, BlendAlpha);
+			
+			BoneConstraint.Weight = BlendWeight;
+		}
+	}
+}
+
 void FIKRetargetRelativeIKOp::InitializePinData(const FIKRetargetProcessor& Processor)
 {
 	// get skeletons we are copying from/to
diff --git a/Engine/Plugins/Experimental/Animation/RelativeIKOp/Source/RelativeIKOp/Public/RelativeIKOp.h b/Engine/Plugins/Experimental/Animation/RelativeIKOp/Source/RelativeIKOp/Public/RelativeIKOp.h
--- a/Engine/Plugins/Experimental/Animation/RelativeIKOp/Source/RelativeIKOp/Public/RelativeIKOp.h
+++ b/Engine/Plugins/Experimental/Animation/RelativeIKOp/Source/RelativeIKOp/Public/RelativeIKOp.h
@@ -140,6 +140,8 @@ struct FRelativeIKFrameBoneConstraint
 	double FeasibilityWeight;
 	FVector ContactOffset;
 	FVector Offset;
+	bool bAmplifyCloseWeight;
+	double SourceDist;
 
 	FRelativeIKFrameBoneConstraint()
 	{
@@ -150,7 +152,11 @@ struct FRelativeIKFrameBoneConstraint
 	{
 		ConstraintType = ERelativeIKFrameBoneConstraintType::Default;
 		Weight = 0.0;
+		FeasibilityWeight = 0.0;
 		Offset = FVector::ZeroVector;
+		ContactOffset = FVector::ZeroVector;
+		bAmplifyCloseWeight = false;
+		SourceDist = 0.0;
 	}
 };
 
@@ -218,6 +224,26 @@ struct FIKRetargetRelativeIKOpSettings : public FIKRetargetOpSettingsBase
 	// Alpha between primary and secondary pair distance relationship
 	UPROPERTY(BlueprintReadWrite, Category = "Parameters")
 	double RetargetSpringAlpha = 0.5;
+	
+	UPROPERTY(EditAnywhere, BlueprintReadWrite, Category = "Importance Weight")
+	bool bApplyImportanceWeighting = true;
+	
+	// TODO: Maybe use IKGoal bones (or detail customization to set automatically)
+	// Amplify solve weight when "important bones" are close to props
+	// UPROPERTY(EditAnywhere, BlueprintReadWrite, Category = "Importance Weight")
+	// TArray<FName> ImportantBones;
+	
+	// Start amplifying weight at this distance (close-ish to prop)
+	UPROPERTY(EditAnywhere, BlueprintReadWrite, Category = "Importance Weight", meta=(ClampMin = "1.0"))
+	double AmplifyCloseThreshold = 15.0;
+
+	// Max weight increase cap distance (very close to prop)
+	UPROPERTY(EditAnywhere, BlueprintReadWrite, Category = "Importance Weight", meta=(ClampMin = "0.1"))
+	double AmplifyCloseDistance = 3.0;
+
+	// Bias feasibility distance weighting relative to distance threshold (+ increase target feasibility/- reduced target feasibility)
+	UPROPERTY(EditAnywhere, BlueprintReadWrite, Category = "Importance Weight", meta=(ClampMin = "1.0", ClampMax = "20.0"))
+	double AmplifyScale = 6.0;
 
 	// Draw source and target body pair relationships for current animation
 	UPROPERTY(EditAnywhere, BlueprintReadWrite, Category = "Debug")
@@ -396,12 +422,15 @@ private:
 
 	void ComputeFramePropSpaces(
 		double SourceScale,
+		const TArray<FName>& AmplifyBones,
 		const TArray<FTransform>& SourcePose,
 		const TArray<FTransform>& TargetPose,
 		int32 PairOffset,
 		TConstArrayView<FVector3f> FramePointView,
 		FDebugRelativeIKDrawInfo& DebugInfo,
 		TArray<FDebugRelativeTargetPairSpace>& DebugPairSpaces);
+	
+	void AmplifyImportanceWeights();
 
 	void ApplyAverageBodyPairTargets(
 		TArray<FIKRigGoal>& OutIKGoals,
