Unshelved from pending changelist '51674342':

PropAnimationModifier Bug Fix

#Synced-CL 51301197

diff --git a/Engine/Plugins/Experimental/Animation/RelativeIKOp/Source/PreviewPropOp/Private/PreviewPropOp.cpp b/Engine/Plugins/Experimental/Animation/RelativeIKOp/Source/PreviewPropOp/Private/PreviewPropOp.cpp
--- a/Engine/Plugins/Experimental/Animation/RelativeIKOp/Source/PreviewPropOp/Private/PreviewPropOp.cpp
+++ b/Engine/Plugins/Experimental/Animation/RelativeIKOp/Source/PreviewPropOp/Private/PreviewPropOp.cpp
@@ -89,6 +89,11 @@ bool FIKRetargetPreviewPropOp::Initialize(
 	for (int32 PropIdx = 0; PropIdx < Settings.PreviewProps.Num(); PropIdx++)
 	{
 		FPreviewPropsData& PreviewData = Settings.PreviewProps[PropIdx];
+		if (!PreviewData.ShowProp)
+		{
+			continue;
+		}
+		
 		if ((!PreviewData.PropStaticMeshAsset && !PreviewData.PropSkeletalMeshAsset) || PreviewData.SourceAttachBone.BoneName == NAME_None)
 		{
 			continue;
diff --git a/Engine/Plugins/Experimental/Animation/RelativeIKOp/Source/PreviewPropOp/Public/PreviewPropOp.h b/Engine/Plugins/Experimental/Animation/RelativeIKOp/Source/PreviewPropOp/Public/PreviewPropOp.h
--- a/Engine/Plugins/Experimental/Animation/RelativeIKOp/Source/PreviewPropOp/Public/PreviewPropOp.h
+++ b/Engine/Plugins/Experimental/Animation/RelativeIKOp/Source/PreviewPropOp/Public/PreviewPropOp.h
@@ -53,6 +53,11 @@ struct FPreviewPropsData
 	GENERATED_BODY()
 
 // #if WITH_EDITORONLY_DATA
+
+	// Prop mesh to anim sequence is looping
+	UPROPERTY(EditAnywhere, BlueprintReadWrite, Category = "Prop", meta=(ReinitializeOnEdit))
+	bool ShowProp = true;
+	
 	// Prop mesh to preview
 	UPROPERTY(EditAnywhere, BlueprintReadWrite, Category = "Prop", meta=(ReinitializeOnEdit))
 	TObjectPtr<UStaticMesh> PropStaticMeshAsset;
diff --git a/Engine/Plugins/Experimental/Animation/RelativeIKOp/Source/RelativeBodyAnimUtils/Private/RelativePropsAnimModifier.cpp b/Engine/Plugins/Experimental/Animation/RelativeIKOp/Source/RelativeBodyAnimUtils/Private/RelativePropsAnimModifier.cpp
--- a/Engine/Plugins/Experimental/Animation/RelativeIKOp/Source/RelativeBodyAnimUtils/Private/RelativePropsAnimModifier.cpp
+++ b/Engine/Plugins/Experimental/Animation/RelativeIKOp/Source/RelativeBodyAnimUtils/Private/RelativePropsAnimModifier.cpp
@@ -454,7 +454,6 @@ void URelativePropsAnimModifier::OnApply_Implementation(UAnimSequence* InAnimati
 					
 					FTransform AttachBoneTransform = UAnimPoseExtensions::GetBonePose(AnimPose, PropAttachBones[PropsIdx], EAnimPoseSpaces::World);
 					FTransform PropsLocalToWorld = PropAttachTfms[PropsIdx] * AttachBoneTransform;
-					const FPositionVertexBuffer& PositionVertexBuffer = PropsData[PropsIdx].GetPositionVertexBuffer();
 					
 					if (PropsData[PropsIdx].PropSkeletalMeshAsset && PropsData[PropsIdx].PropAnimSequence)
 					{
@@ -571,12 +570,15 @@ void URelativePropsAnimModifier::GetPropsVertices(TArray<TArray<FVector3f>>& Pro
 		
 		const FTransform AttachBoneTransform = UAnimPoseExtensions::GetBonePose(AnimPose, PropAttachBones[PropsIdx], EAnimPoseSpaces::World);
 		
-		const FPositionVertexBuffer& PositionVertexBuffer = PropsData[PropsIdx].GetPositionVertexBuffer();
-		
-		int32 NumVerts = PositionVertexBuffer.GetNumVertices();
+		// const FPositionVertexBuffer& PositionVertexBuffer = PropsData[PropsIdx].GetPositionVertexBuffer();
+		// int32 NumVerts = PositionVertexBuffer.GetNumVertices();
+		const FMeshDescription* MeshDescriptionPtr = PropsData[PropsIdx].GetMeshDescription(0);
+		const int32 NumVerts = MeshDescriptionPtr->Vertices().Num();
+
 		for (int32 i = 0; i < NumVerts; ++i)
 		{
-			FVector LocalPos = PropAttachTfms[PropsIdx].TransformPosition(FVector(PositionVertexBuffer.VertexPosition(i)));
+			// FVector LocalPos = PropAttachTfms[PropsIdx].TransformPosition(FVector(PositionVertexBuffer.VertexPosition(i)));
+			FVector LocalPos = PropAttachTfms[PropsIdx].TransformPosition(FVector(MeshDescriptionPtr->GetVertexPosition(i)));
 			PropsVLocations[PropsIdx].Add(FVector3f(AttachBoneTransform.TransformPosition(LocalPos)));
 		}
 	}
