Some small fix for RIK

#Synced-CL 51694918

diff --git a/Engine/Plugins/Experimental/Animation/RelativeIKOp/Source/RelativeIKOp/Private/RelativeIKOp.cpp b/Engine/Plugins/Experimental/Animation/RelativeIKOp/Source/RelativeIKOp/Private/RelativeIKOp.cpp
--- a/Engine/Plugins/Experimental/Animation/RelativeIKOp/Source/RelativeIKOp/Private/RelativeIKOp.cpp
+++ b/Engine/Plugins/Experimental/Animation/RelativeIKOp/Source/RelativeIKOp/Private/RelativeIKOp.cpp
@@ -482,7 +482,7 @@ void FIKRetargetRelativeIKOp::ComputeFramePropSpaces(
 			int32 PinPropIdx = CachePinSourcePropBones.Find(SourcePropBone);
 			if (PinPropIdx == INDEX_NONE)
 			{
-				return;
+				continue;
 			}
 			
 			int32 SourcePropPoseIdx = CacheSourceSkelIndices[SourcePropBone];
@@ -1153,7 +1153,7 @@ void FIKRetargetRelativeIKOp::ApplyMatrixBodyPairTargetsWithPropScale(
 		OutBoneLocations.Add(TargetBonePose.GetLocation());
 		
 		int32 PinPropIdx = CachePinSourcePropBones.Find(SourceBoneName);
-		if (PinPropIdx != INDEX_NONE)
+		if (PinPropIdx != INDEX_NONE && CachePropId2PinBoneDataId[PinPropIdx] != INDEX_NONE)
 		{
 			ScalarIndexMap.Add(OutBoneIdx);
 	        ScalarBoneNames.Add(SourceBoneName);
@@ -1336,7 +1336,7 @@ void FIKRetargetRelativeIKOp::ApplyMatrixBodyPairTargetsWithPropScale(
 		}
 		
 		int32 CachePropIdx = CachePinSourcePropBones.Find(SourceBoneName);
-		if (CachePropIdx == INDEX_NONE)
+		if (CachePropIdx == INDEX_NONE || CachePropId2PinBoneDataId[CachePropIdx] == INDEX_NONE)
 		{
 			continue;
 		}
@@ -2413,8 +2413,8 @@ void FIKRetargetRelativeIKOp::UpdateCachePropNotifyInfo(URelativePropsBakeAnimNo
 		CachePinSourcePropBones.AddUnique(PropBone);
 	}
 	
-	CachePropId2PinTargetPropSettingId.SetNumZeroed(CachePinSourcePropBones.Num());
-	CachePropId2PinBoneDataId.SetNumZeroed(CachePinSourcePropBones.Num());
+	CachePropId2PinTargetPropSettingId.Init(INDEX_NONE, CachePinSourcePropBones.Num());
+	CachePropId2PinBoneDataId.Init(INDEX_NONE, CachePinSourcePropBones.Num());
 	int32 PinBoneId = 0;
 	for (int32 SettingId = 0; SettingId < Settings.PinPropBones.Num(); ++SettingId)
 	{
diff --git a/Engine/Plugins/Experimental/Animation/RelativeIKOp/Source/RelativeIKOp/Public/RelativeIKOp.h b/Engine/Plugins/Experimental/Animation/RelativeIKOp/Source/RelativeIKOp/Public/RelativeIKOp.h
--- a/Engine/Plugins/Experimental/Animation/RelativeIKOp/Source/RelativeIKOp/Public/RelativeIKOp.h
+++ b/Engine/Plugins/Experimental/Animation/RelativeIKOp/Source/RelativeIKOp/Public/RelativeIKOp.h
@@ -211,7 +211,7 @@ struct FIKRetargetRelativeIKOpSettings : public FIKRetargetOpSettingsBase
 
 	// Frames of temporal smoothing of body pair verts (0 is no smoothing)
 	UPROPERTY(EditAnywhere, BlueprintReadWrite, Category = "Parameters", meta=(ClampMin = "0", ClampMax = "60"))
-	int32 TemporalSmoothingRadius = 15;
+	int32 TemporalSmoothingRadius = 5;
 
 	// Ignore source scaling when computing relative distance relationships
 	UPROPERTY(EditAnywhere, BlueprintReadWrite, Category = "Parameters")
