PropBake Fix

#Synced-CL 51875096

diff --git a/Engine/Plugins/Experimental/Animation/RelativeIKOp/Source/RelativeBodyAnimUtils/Private/RelativePropsAnimModifier.cpp b/Engine/Plugins/Experimental/Animation/RelativeIKOp/Source/RelativeBodyAnimUtils/Private/RelativePropsAnimModifier.cpp
--- a/Engine/Plugins/Experimental/Animation/RelativeIKOp/Source/RelativeBodyAnimUtils/Private/RelativePropsAnimModifier.cpp
+++ b/Engine/Plugins/Experimental/Animation/RelativeIKOp/Source/RelativeBodyAnimUtils/Private/RelativePropsAnimModifier.cpp
@@ -67,6 +67,7 @@ struct FPropsMeshAndTree
 {
 	UE::Geometry::FDynamicMesh3 Mesh;
 	UE::Geometry::TMeshAABBTree3<UE::Geometry::FDynamicMesh3> AABBTree;
+	TArray<int32> DynVidToSourceVid;  // Maps DynamicMesh VID → MeshDesc VID
 
 	FPropsMeshAndTree()
 		: AABBTree(&Mesh, false) {}  // defer build
@@ -75,6 +76,37 @@ struct FPropsMeshAndTree
 	{
 		FMeshDescriptionToDynamicMesh Converter;
 		Converter.Convert(MeshDesc, Mesh);
+		// Build mapping: identity for original vertices,
+		// source vertex for split (bowtie) vertices.
+		// FMeshDescriptionToDynamicMesh internally tracks this mapping
+		// during non-manifold topology repair.
+		int32 MaxVid = Mesh.MaxVertexID();
+		int32 OriginalCount = MeshDesc->Vertices().Num();
+		DynVidToSourceVid.SetNumUninitialized(MaxVid);
+		for (int32 Vid : Mesh.VertexIndicesItr())
+		{
+			// Original vertices map to themselves; split vertices
+			// share the same position as their source — find by position match
+			// or use the converter's internal VtxMap if accessible.
+			DynVidToSourceVid[Vid] = (Vid < OriginalCount) ? Vid : INDEX_NONE;
+		}
+		// For split vertices (Vid >= OriginalCount), find closest original vertex
+		// by position to establish the mapping.
+		for (int32 Vid : Mesh.VertexIndicesItr())
+		{
+			if (Vid >= OriginalCount)
+			{
+				FVector3d SplitPos = Mesh.GetVertex(Vid);
+				for (int32 OrigVid = 0; OrigVid < OriginalCount; ++OrigVid)
+				{
+					if (Mesh.GetVertex(OrigVid).Equals(SplitPos, SMALL_NUMBER))
+					{
+						DynVidToSourceVid[Vid] = OrigVid;
+						break;
+					}
+				}
+			}
+		}
 		AABBTree.Build();
 	}
 	
@@ -83,11 +115,12 @@ struct FPropsMeshAndTree
 	{
 		for (int32 Vid : Mesh.VertexIndicesItr())
 		{
-			// FVector3d OldPos = Mesh.GetVertex(Vid);
-			Mesh.SetVertex(Vid, FVector3d(VLocation[Vid]));
+			int32 Src = (Vid < DynVidToSourceVid.Num()) ? DynVidToSourceVid[Vid] : Vid;
+			if (Src >= 0 && Src < VLocation.Num())
+			{
+				Mesh.SetVertex(Vid, FVector3d(VLocation[Src]));
+			}
 		}
-
-		// After modifying vertices, rebuild the tree
 		AABBTree.Build();
 	}
 };
