RelativeIK: Prop Pushout ops

#Synced-CL 52116586

diff --git a/Engine/Plugins/Experimental/Animation/RelativeIKOp/Source/BodyIntersectIKOp/Private/BodyIntersectIKOp.cpp b/Engine/Plugins/Experimental/Animation/RelativeIKOp/Source/BodyIntersectIKOp/Private/BodyIntersectIKOp.cpp
--- a/Engine/Plugins/Experimental/Animation/RelativeIKOp/Source/BodyIntersectIKOp/Private/BodyIntersectIKOp.cpp
+++ b/Engine/Plugins/Experimental/Animation/RelativeIKOp/Source/BodyIntersectIKOp/Private/BodyIntersectIKOp.cpp
@@ -141,7 +141,7 @@ void FIKRetargetBodyIntersectIKOp::Run(
 			continue;
 		}
 		
-		FKShapeElem* GoalShape = FPhysShapeUtils::FindBodyShape(PhysicsAsset, IKGoal->BoneName);
+		const FKShapeElem* GoalShape = FPhysShapeUtils::FindBodyShape(PhysicsAsset, IKGoal->BoneName);
 		double GoalRadius = FPhysShapeUtils::CalcShapeSmallRadius(GoalShape); // * GoalSettings.SphereScale;
 		if (FMath::IsNearlyZero(GoalRadius))
 		{
@@ -215,7 +215,7 @@ void FIKRetargetBodyIntersectIKOp::Run(
 				continue;
 			}
 			
-			FKShapeElem* GoalShape = FPhysShapeUtils::FindBodyShape(PhysicsAsset, IKGoal->BoneName);
+			const FKShapeElem* GoalShape = FPhysShapeUtils::FindBodyShape(PhysicsAsset, IKGoal->BoneName);
 			double GoalRadius = FPhysShapeUtils::CalcShapeSmallRadius(GoalShape); // * GoalSettings.SphereScale;
 			if (FMath::IsNearlyZero(GoalRadius))
 			{
@@ -436,7 +436,7 @@ void FPoleVectorChainIntersector::ApplyIntersectRotation(
 		// TODO: Should we copy these settings into the struct?
 		const FIKPoleVectorIntersectSettings& BodySettings = Settings->PoleVectorIntersectSettings[SettingsIdx];
 
-		FKShapeElem* BodyShape = FPhysShapeUtils::FindBodyShape(PhysicsAsset, BodySettings.BodyName);
+		const FKShapeElem* BodyShape = FPhysShapeUtils::FindBodyShape(PhysicsAsset, BodySettings.BodyName);
 		double BodyRadius = FPhysShapeUtils::CalcShapeSmallRadius(BodyShape) * BodySettings.SphereScale;
 		if (FMath::IsNearlyZero(BodyRadius))
 		{
@@ -503,7 +503,7 @@ void FPoleVectorChainIntersector::ApplyIntersectRotation(
 			
 			const FIKPoleVectorIntersectSettings& BodySettings = Settings->PoleVectorIntersectSettings[SettingsIdx];
 
-			FKShapeElem* BodyShape = FPhysShapeUtils::FindBodyShape(PhysicsAsset, BodySettings.BodyName);
+			const FKShapeElem* BodyShape = FPhysShapeUtils::FindBodyShape(PhysicsAsset, BodySettings.BodyName);
 			double BodyRadius = FPhysShapeUtils::CalcShapeSmallRadius(BodyShape) * BodySettings.SphereScale;
 			if (FMath::IsNearlyZero(BodyRadius))
 			{
@@ -544,7 +544,7 @@ FVector FPoleVectorChainIntersector::GetIntersectDelta(
 		}
 		
 		FTransform BodyTfm = TargetGlobalPose[BoneIdx];
-		FKShapeElem* ShapeElem = FPhysShapeUtils::FindBodyShape(PhysicsAsset, BodyName);
+		const FKShapeElem* ShapeElem = FPhysShapeUtils::FindBodyShape(PhysicsAsset, BodyName);
 		if (!ShapeElem)
 		{
 			continue;
@@ -682,7 +682,7 @@ void FIKRetargetBodyIntersectIKOp::InitPoleVectorIntersectors(
 
 FVector FIKRetargetBodyIntersectIKOp::GetTotalIntersectDelta(
 	int32 CollisionProxyIndex,
-	FKShapeElem* IntersectShape,
+	const FKShapeElem* IntersectShape,
 	const FTransform& IntersectTfm,
 	const TArray<FTransform>& TargetGlobalPose,
 	const FTargetSkeleton& TargetSkeleton,
@@ -717,7 +717,7 @@ FVector FIKRetargetBodyIntersectIKOp::GetTotalIntersectDelta(
 		}
 		
 		FTransform BodyTfm = TargetGlobalPose[BodyPoseIdx];
-		FKShapeElem* BodyShape = FPhysShapeUtils::FindBodyShape(PhysicsAsset, BodyName);
+		const FKShapeElem* BodyShape = FPhysShapeUtils::FindBodyShape(PhysicsAsset, BodyName);
 		if (!BodyShape)
 		{
 			continue;
diff --git a/Engine/Plugins/Experimental/Animation/RelativeIKOp/Source/BodyIntersectIKOp/Private/PhysicsBodyHelpers.cpp b/Engine/Plugins/Experimental/Animation/RelativeIKOp/Source/BodyIntersectIKOp/Private/PhysicsBodyHelpers.cpp
--- a/Engine/Plugins/Experimental/Animation/RelativeIKOp/Source/BodyIntersectIKOp/Private/PhysicsBodyHelpers.cpp
+++ b/Engine/Plugins/Experimental/Animation/RelativeIKOp/Source/BodyIntersectIKOp/Private/PhysicsBodyHelpers.cpp
@@ -5,7 +5,7 @@
 #include "PhysicsEngine/PhysicsAsset.h"
 #include "PhysicsEngine/SkeletalBodySetup.h"
 
-double FBodyIntersectUtils::CalcIntersectionPairDelta(const FTransform& GoalShapeTfm, FKShapeElem* GoalShapeElem, const FTransform& TargetShapeTfm, FKShapeElem* TargetShapeElem, FVector& OutDeltaDir)
+double FBodyIntersectUtils::CalcIntersectionPairDelta(const FTransform& GoalShapeTfm, const FKShapeElem* GoalShapeElem, const FTransform& TargetShapeTfm, const FKShapeElem* TargetShapeElem, FVector& OutDeltaDir)
 {
 	if (!GoalShapeElem || !TargetShapeElem)
 	{
@@ -20,8 +20,8 @@ double FBodyIntersectUtils::CalcIntersectionPairDelta(const FTransform& GoalShap
 	{
 	case EBodyIntersectPairType::SphereSphere:
 		{
-			FKSphereElem* SphereElem1 = static_cast<FKSphereElem*>(GoalShapeElem);
-			FKSphereElem* SphereElem2 = static_cast<FKSphereElem*>(TargetShapeElem);
+			const FKSphereElem* SphereElem1 = static_cast<const FKSphereElem*>(GoalShapeElem);
+			const FKSphereElem* SphereElem2 = static_cast<const FKSphereElem*>(TargetShapeElem);
 			Dist = CalcSphereSphereShapeDelta(OutDeltaDir, GoalShapeTfm, SphereElem1, TargetShapeTfm, SphereElem2);
 			break;
 		}
@@ -30,15 +30,15 @@ double FBodyIntersectUtils::CalcIntersectionPairDelta(const FTransform& GoalShap
 		{
 			if (NeedInverse)
 			{
-				FKSphereElem* SphereElem1 = static_cast<FKSphereElem*>(TargetShapeElem);
-				FKBoxElem* BoxElem2 = static_cast<FKBoxElem*>(GoalShapeElem);
+				const FKSphereElem* SphereElem1 = static_cast<const FKSphereElem*>(TargetShapeElem);
+				const FKBoxElem* BoxElem2 = static_cast<const FKBoxElem*>(GoalShapeElem);
 				Dist = CalcSphereBoxShapeDelta(OutDeltaDir, TargetShapeTfm, SphereElem1, GoalShapeTfm, BoxElem2);
 				OutDeltaDir *= -1;
 			}
 			else
 			{
-				FKSphereElem* SphereElem1 = static_cast<FKSphereElem*>(GoalShapeElem);
-				FKBoxElem* BoxElem2 = static_cast<FKBoxElem*>(TargetShapeElem);
+				const FKSphereElem* SphereElem1 = static_cast<const FKSphereElem*>(GoalShapeElem);
+				const FKBoxElem* BoxElem2 = static_cast<const FKBoxElem*>(TargetShapeElem);
 				Dist = CalcSphereBoxShapeDelta(OutDeltaDir, GoalShapeTfm, SphereElem1, TargetShapeTfm, BoxElem2);
 			}
 			break;
@@ -48,15 +48,15 @@ double FBodyIntersectUtils::CalcIntersectionPairDelta(const FTransform& GoalShap
 		{
 			if (NeedInverse)
 			{
-				FKSphereElem* SphereElem1 = static_cast<FKSphereElem*>(TargetShapeElem);
-				FKSphylElem* CapsuleElem2 = static_cast<FKSphylElem*>(GoalShapeElem);
+				const FKSphereElem* SphereElem1 = static_cast<const FKSphereElem*>(TargetShapeElem);
+				const FKSphylElem* CapsuleElem2 = static_cast<const FKSphylElem*>(GoalShapeElem);
 				Dist = CalcSphereCapsuleShapeDelta(OutDeltaDir, TargetShapeTfm, SphereElem1, GoalShapeTfm, CapsuleElem2);
 				OutDeltaDir *= -1;
 			}
 			else
 			{
-				FKSphereElem* SphereElem1 = static_cast<FKSphereElem*>(GoalShapeElem);
-				FKSphylElem* CapsuleElem2 = static_cast<FKSphylElem*>(TargetShapeElem);
+				const FKSphereElem* SphereElem1 = static_cast<const FKSphereElem*>(GoalShapeElem);
+				const FKSphylElem* CapsuleElem2 = static_cast<const FKSphylElem*>(TargetShapeElem);
 				Dist = CalcSphereCapsuleShapeDelta(OutDeltaDir, GoalShapeTfm, SphereElem1, TargetShapeTfm, CapsuleElem2);
 			}
 			break;
@@ -74,15 +74,15 @@ double FBodyIntersectUtils::CalcIntersectionPairDelta(const FTransform& GoalShap
 		{
 			if (NeedInverse)
 			{
-				FKBoxElem* BoxElem1 = static_cast<FKBoxElem*>(TargetShapeElem);
-				FKSphylElem* CapsuleElem2 = static_cast<FKSphylElem*>(GoalShapeElem);
+				const FKBoxElem* BoxElem1 = static_cast<const FKBoxElem*>(TargetShapeElem);
+				const FKSphylElem* CapsuleElem2 = static_cast<const FKSphylElem*>(GoalShapeElem);
 				Dist = CalcBoxCapsuleShapeDelta(OutDeltaDir, TargetShapeTfm, BoxElem1, GoalShapeTfm, CapsuleElem2);
 				OutDeltaDir *= -1;
 			}
 			else
 			{
-				FKBoxElem* BoxElem1 = static_cast<FKBoxElem*>(GoalShapeElem);
-				FKSphylElem* CapsuleElem2 = static_cast<FKSphylElem*>(TargetShapeElem);
+				const FKBoxElem* BoxElem1 = static_cast<const FKBoxElem*>(GoalShapeElem);
+				const FKSphylElem* CapsuleElem2 = static_cast<const FKSphylElem*>(TargetShapeElem);
 				Dist = CalcBoxCapsuleShapeDelta(OutDeltaDir, GoalShapeTfm, BoxElem1, TargetShapeTfm, CapsuleElem2);
 			}
 			break;
@@ -90,8 +90,8 @@ double FBodyIntersectUtils::CalcIntersectionPairDelta(const FTransform& GoalShap
 		
 	case EBodyIntersectPairType::CapsuleCapsule:
 		{
-			FKSphylElem* CapsuleElem1 = static_cast<FKSphylElem*>(GoalShapeElem);
-			FKSphylElem* CapsuleElem2 = static_cast<FKSphylElem*>(TargetShapeElem);
+			const FKSphylElem* CapsuleElem1 = static_cast<const FKSphylElem*>(GoalShapeElem);
+			const FKSphylElem* CapsuleElem2 = static_cast<const FKSphylElem*>(TargetShapeElem);
 			Dist = CalcCapsuleCapsuleShapeDelta(OutDeltaDir, GoalShapeTfm, CapsuleElem1, TargetShapeTfm, CapsuleElem2);
 			break;
 		}
@@ -106,6 +106,41 @@ double FBodyIntersectUtils::CalcIntersectionPairDelta(const FTransform& GoalShap
 	return Dist;
 }
 
+double FBodyIntersectUtils::CalcIntersectionPairDelta(const FTransform& GoalShapeTfm, const FKAggregateGeom* GoalAggGeom, const FTransform& TargetShapeTfm, const FKAggregateGeom* TargetAggGeom, FVector& OutDeltaDir)
+{
+	if (!GoalAggGeom || !TargetAggGeom)
+	{
+		return -1.0;
+	}
+	
+	double Dist = -1.0;
+	FVector Delta = FVector::ZeroVector;
+	FTransform UpdateTfm = GoalShapeTfm;
+	int32 NumGoalShapes = GoalAggGeom->GetElementCount();
+	int32 NumTargetShapes = TargetAggGeom->GetElementCount();
+	for (int32 GoalShapeIdx = 0; GoalShapeIdx < NumGoalShapes; GoalShapeIdx++)
+	{
+		const FKShapeElem* GoalShapeElem = GoalAggGeom->GetElement(GoalShapeIdx);
+		for (int32 TargetShapeIdx = 0; TargetShapeIdx < NumTargetShapes; TargetShapeIdx++)
+		{
+			const FKShapeElem* TargetShapeElem = TargetAggGeom->GetElement(TargetShapeIdx);
+			
+			FVector PairDeltaDir = FVector::ZeroVector;
+			double PairDist = CalcIntersectionPairDelta(UpdateTfm, GoalShapeElem, TargetShapeTfm, TargetShapeElem, PairDeltaDir);
+			
+			if (PairDist > 0)
+			{
+				Delta += PairDist*PairDeltaDir;
+				UpdateTfm.SetTranslation(GoalShapeTfm.GetTranslation() + Delta);
+			}
+		}
+	}
+	
+	Dist = Delta.Size();
+	OutDeltaDir = (Dist > UE_SMALL_NUMBER) ? Delta / Dist : FVector::ZeroVector;
+	return Dist;
+}
+
 double FBodyIntersectUtils::SegmentDistToAABB(
 	const FVector& A1, const FVector& A2,
 	const FVector& BoxMin, const FVector& BoxMax,
@@ -391,7 +426,7 @@ EBodyIntersectPairType FBodyIntersectUtils::CalcBodyIntersectPairType(const FKSh
 	return EBodyIntersectPairType::Unknown;
 }
 
-double FPhysShapeUtils::CalcShapeSmallRadius(FKShapeElem* ShapeElem)
+double FPhysShapeUtils::CalcShapeSmallRadius(const FKShapeElem* ShapeElem)
 {
 	// Return 1.0 as default so scaling can set "size" of sphere
 	const double DefaultRadius = 1.0;
@@ -404,18 +439,18 @@ double FPhysShapeUtils::CalcShapeSmallRadius(FKShapeElem* ShapeElem)
 	{
 	case EAggCollisionShape::Sphere:
 		{
-			FKSphereElem* SphereElem = static_cast<FKSphereElem*>(ShapeElem);
+			const FKSphereElem* SphereElem = static_cast<const FKSphereElem*>(ShapeElem);
 			return SphereElem->Radius;
 		}
 
 	case EAggCollisionShape::Sphyl:
 		{
-			FKSphylElem* CapsuleElem = static_cast<FKSphylElem*>(ShapeElem);
+			const FKSphylElem* CapsuleElem = static_cast<const FKSphylElem*>(ShapeElem);
 			return CapsuleElem->Radius;
 		}
 	case EAggCollisionShape::Box:
 		{
-			FKBoxElem* BoxElem = static_cast<FKBoxElem*>(ShapeElem);
+			const FKBoxElem* BoxElem = static_cast<const FKBoxElem*>(ShapeElem);
 			return 0.5 * FMath::Min(FMath::Min(BoxElem->X, BoxElem->Y), BoxElem->Z);
 		}
 		// case EAggCollisionShape::Convex:
@@ -431,19 +466,19 @@ double FPhysShapeUtils::CalcShapeSmallRadius(FKShapeElem* ShapeElem)
 	}
 }
 
-double FPhysShapeUtils::CalcShapeAvgRadius(FKShapeElem* ShapeElem)
+double FPhysShapeUtils::CalcShapeAvgRadius(const FKShapeElem* ShapeElem)
 {
 	switch (ShapeElem->GetShapeType())
 	{
 	case EAggCollisionShape::Sphere:
 		{
-			FKSphereElem* SphereElem = static_cast<FKSphereElem*>(ShapeElem);
+			const FKSphereElem* SphereElem = static_cast<const FKSphereElem*>(ShapeElem);
 			return SphereElem->Radius;
 		}
 
 	case EAggCollisionShape::Sphyl:
 		{
-			FKSphylElem* CapsuleElem = static_cast<FKSphylElem*>(ShapeElem);
+			const FKSphylElem* CapsuleElem = static_cast<const FKSphylElem*>(ShapeElem);
 			return CapsuleElem->Radius + 0.25*CapsuleElem->Length;
 		}
 
@@ -454,7 +489,7 @@ double FPhysShapeUtils::CalcShapeAvgRadius(FKShapeElem* ShapeElem)
 	}
 }
 
-FKShapeElem* FPhysShapeUtils::FindBodyShape(const UPhysicsAsset* PhysAsset, FName BoneName)
+const FKShapeElem* FPhysShapeUtils::FindBodyShape(const UPhysicsAsset* PhysAsset, FName BoneName)
 {
 	int32 BodyIdx = PhysAsset->FindBodyIndex(BoneName);
 	if (BodyIdx == INDEX_NONE)
@@ -464,3 +499,14 @@ FKShapeElem* FPhysShapeUtils::FindBodyShape(const UPhysicsAsset* PhysAsset, FNam
 	}
 	return PhysAsset->SkeletalBodySetups[BodyIdx]->AggGeom.GetElement(0);
 }
+
+const FKAggregateGeom* FPhysShapeUtils::FindBodyAggGeom(const UPhysicsAsset* PhysAsset, FName BoneName)
+{
+	int32 BodyIdx = PhysAsset->FindBodyIndex(BoneName);
+	if (BodyIdx == INDEX_NONE)
+	{
+		UE_LOG(LogAnimation, Warning, TEXT("No body index found: %s"), *BoneName.ToString());
+		return nullptr;
+	}
+	return &PhysAsset->SkeletalBodySetups[BodyIdx]->AggGeom;
+}
diff --git a/Engine/Plugins/Experimental/Animation/RelativeIKOp/Source/BodyIntersectIKOp/Private/PhysicsBodyHelpers.h b/Engine/Plugins/Experimental/Animation/RelativeIKOp/Source/BodyIntersectIKOp/Private/PhysicsBodyHelpers.h
deleted file mode 100644
diff --git a/Engine/Plugins/Experimental/Animation/RelativeIKOp/Source/BodyIntersectIKOp/Public/BodyIntersectIKOp.h b/Engine/Plugins/Experimental/Animation/RelativeIKOp/Source/BodyIntersectIKOp/Public/BodyIntersectIKOp.h
--- a/Engine/Plugins/Experimental/Animation/RelativeIKOp/Source/BodyIntersectIKOp/Public/BodyIntersectIKOp.h
+++ b/Engine/Plugins/Experimental/Animation/RelativeIKOp/Source/BodyIntersectIKOp/Public/BodyIntersectIKOp.h
@@ -234,7 +234,7 @@ private:
 	
 	FVector GetTotalIntersectDelta(
 		int32 CollisionProxyIndex,
-		FKShapeElem* IntersectShape,
+		const FKShapeElem* IntersectShape,
 		const FTransform& IntersectTfm,
 		const TArray<FTransform>& TargetGlobalPose,
 		const FTargetSkeleton& TargetSkeleton,
diff --git a/Engine/Plugins/Experimental/Animation/RelativeIKOp/Source/BodyIntersectIKOp/Public/PhysicsBodyHelpers.h b/Engine/Plugins/Experimental/Animation/RelativeIKOp/Source/BodyIntersectIKOp/Public/PhysicsBodyHelpers.h
new file mode 100644
--- /dev/null
+++ b/Engine/Plugins/Experimental/Animation/RelativeIKOp/Source/BodyIntersectIKOp/Public/PhysicsBodyHelpers.h
@@ -0,0 +1,60 @@
+﻿// Copyright Epic Games, Inc. All Rights Reserved.
+
+#pragma once
+
+#include "Math/Vector.h"
+#include "Math/Transform.h"
+
+#include "PhysicsEngine/PhysicsAsset.h"
+#include "PhysicsEngine/SkeletalBodySetup.h"
+
+
+struct FKAggregateGeom;
+struct FKShapeElem;
+struct FKSphereElem;
+struct FKBoxElem;
+struct FKSphylElem;
+class UPhysicsAsset;
+
+enum class EBodyIntersectPairType : uint8 
+{
+	SphereSphere,
+	SphereBox,
+	BoxBox,
+	SphereCapsule,
+	BoxCapsule,
+	CapsuleCapsule,
+	
+	Unknown
+};
+
+/**
+ * Helper static class providing physics body intersection delta functions
+ */
+class BODYINTERSECTIKOP_API FBodyIntersectUtils
+{
+public:
+	static double CalcIntersectionPairDelta(const FTransform& GoalShapeTfm, const FKShapeElem* GoalShapeElem, const FTransform& TargetShapeTfm, const FKShapeElem* TargetShapeElem, FVector& OutDeltaDir);
+	static double CalcIntersectionPairDelta(const FTransform& GoalShapeTfm, const FKAggregateGeom* GoalAggGeom, const FTransform& TargetShapeTfm, const FKAggregateGeom* TargetAggGeom, FVector& OutDeltaDir);
+private:
+	static EBodyIntersectPairType CalcBodyIntersectPairType(const FKShapeElem* ShapeElem1, const FKShapeElem* ShapeElem2, bool &NeedInverse);
+	
+	static double CalcSphereSphereShapeDelta(FVector& OutDeltaDir, const FTransform& ShapeTfm1, const FKSphereElem* ShapeElem1, const FTransform& ShapeTfm2, const FKSphereElem* ShapeElem2);
+	static double CalcSphereBoxShapeDelta(FVector& OutDeltaDir, const FTransform& ShapeTfm1, const FKSphereElem* ShapeElem1, const FTransform& ShapeTfm2, const FKBoxElem* ShapeElem2);
+	static double CalcSphereCapsuleShapeDelta(FVector& OutDeltaDir, const FTransform& ShapeTfm1, const FKSphereElem* ShapeElem1, const FTransform& ShapeTfm2, const FKSphylElem* ShapeElem2);
+	static double CalcBoxBoxShapeDelta(FVector& OutDeltaDir, const FTransform& ShapeTfm1, const FKBoxElem* ShapeElem1, const FTransform& ShapeTfm2, const FKBoxElem* ShapeElem2);
+	static double CalcBoxCapsuleShapeDelta(FVector& OutDeltaDir, const FTransform& ShapeTfm1, const FKBoxElem* ShapeElem1, const FTransform& ShapeTfm2, const FKSphylElem* ShapeElem2);
+	static double CalcCapsuleCapsuleShapeDelta(FVector& OutDeltaDir, const FTransform& ShapeTfm1, const FKSphylElem* ShapeElem1, const FTransform& ShapeTfm2, const FKSphylElem* ShapeElem2);
+
+	static double SegmentDistToAABB(const FVector& A, const FVector& B, const FVector& BoxMin, const FVector& BoxMax, FVector& OutSeg, FVector& OutBox);
+};
+
+class BODYINTERSECTIKOP_API FPhysShapeUtils
+{
+public:
+	static double CalcShapeSmallRadius(const FKShapeElem* ShapeElem);
+	static double CalcShapeAvgRadius(const FKShapeElem* ShapeElem);
+
+	static const FKShapeElem* FindBodyShape(const UPhysicsAsset* PhysAsset, FName BoneName);
+	static const FKAggregateGeom* FindBodyAggGeom(const UPhysicsAsset* PhysAsset, FName BoneName);
+};
diff --git a/Engine/Plugins/Experimental/Animation/RelativeIKOp/Source/RelativeIKOp/Private/RelativeIKOp.cpp b/Engine/Plugins/Experimental/Animation/RelativeIKOp/Source/RelativeIKOp/Private/RelativeIKOp.cpp
--- a/Engine/Plugins/Experimental/Animation/RelativeIKOp/Source/RelativeIKOp/Private/RelativeIKOp.cpp
+++ b/Engine/Plugins/Experimental/Animation/RelativeIKOp/Source/RelativeIKOp/Private/RelativeIKOp.cpp
@@ -4,6 +4,7 @@
 
 #include "IKRigDebugRendering.h"
 #include "RelativeBodyAnimNotifies.h"
+#include "PhysicsBodyHelpers.h"
 #include "Animation/AnimInstance.h"
 #include "Animation/AnimMontage.h"
 #include "Animation/AnimSequence.h"
@@ -91,20 +92,28 @@ bool FIKRetargetRelativeIKOp::Initialize(
 		return false;
 	}
 	
-	TargetPhysicsAsset = Settings.TargetPhysicsAssetOverride;
-
-	if (!TargetPhysicsAsset)
+	TargetRelativeIKPhysicsAsset = Settings.TargetRelativeIKPhysicsAssetOverride;
+	TargetIntersectOpPhysicsAsset= Settings.TargetIntersectOpPhysicsAssetOverride;
+	if (!TargetRelativeIKPhysicsAsset)
 	{
 		InLog.LogWarning( FText::Format(
 		LOCTEXT("MissingPhysicsAssets", "{0}: Both Source and Target Physics Assets must be specified. "), FText::FromName(GetName())));
 		return false;
 	}
+	if (Settings.bEnableIntersectPushOut && !TargetIntersectOpPhysicsAsset)
+	{
+		TargetIntersectOpPhysicsAsset = TargetRelativeIKPhysicsAsset;
+		InLog.LogWarning( FText::Format(
+		LOCTEXT("MissingPhysicsAssets", "Missing TargetIntersectOpPhysicsAsset Override, use TargetRelativeIKPhysicsAsset Override instead!"), FText::FromName(GetName())));
+	}
 
 	UpdateCacheBoneChains(InProcessor, ParentRigOp);
 	UpdateCacheSkelInfo(InSourceSkeleton, InTargetSkeleton);
 	
 	// Add all pin bones to the source/target skel lookup tables
 	InitializePinData(InProcessor);
+	
+	InitializeIntersectionData(InProcessor);
 
 	// Force update cache info if new anim montage found
 	MontageInstance = nullptr;
@@ -219,6 +228,10 @@ void FIKRetargetRelativeIKOp::Run(
 		ApplyAverageBodyPairTargets(IKGoals, OutTargetGlobalPose, LocalDebugInfo, LocalTargetPairSpaces);
 	}
 
+	if (Settings.bEnableIntersectPushOut)
+	{
+		ApplyIntersectPushOut(IKGoals, OutTargetGlobalPose, LocalDebugInfo);
+	}
 #if WITH_EDITOR
 	LocalDebugInfo.TargetDomainBones.Reset(CacheSourceEffectBones.Num());
 	for (int PairIdx=0; PairIdx<NumMeshPairs; PairIdx++)
@@ -799,6 +812,26 @@ void FIKRetargetRelativeIKOp::InitializePinData(const FIKRetargetProcessor& Proc
 	}
 }
 
+void FIKRetargetRelativeIKOp::InitializeIntersectionData(const FIKRetargetProcessor& Processor)
+{
+	// get skeletons we are copying from/to
+	const FRetargetSkeleton& SourceSkeleton = Processor.GetSkeleton(ERetargetSourceOrTarget::Source);
+	const FRetargetSkeleton& TargetSkeleton = Processor.GetSkeleton(ERetargetSourceOrTarget::Target);
+	
+	for (FName IntersectBody: Settings.IntersectBodies)
+	{
+		FName SourceBoneName = IntersectBody;
+		FName TargetBoneName = ApplyBodyMap(SourceBoneName);
+		if (TargetBoneName == NAME_None)
+		{
+			continue;
+		}
+		
+		UpdateSourceBoneMap(SourceBoneName);
+		UpdateTargetBoneMap(TargetBoneName);
+	}
+}
+
 void FIKRetargetRelativeIKOp::CachePinBoneTransforms(
 	const FIKRetargetProcessor& Processor,
 	const TArray<FTransform>& InSourceGlobalPose,
@@ -1393,6 +1426,247 @@ void FIKRetargetRelativeIKOp::ApplyMatrixBodyPairTargetsWithPropScale(
 	}
 }
 
+void FIKRetargetRelativeIKOp::ApplyIntersectPushOut(
+	TArray<FIKRigGoal>& OutIKGoals,
+	const TArray<FTransform>& TargetPose,
+	FDebugRelativeIKDrawInfo& DebugInfo)
+{
+	const UPhysicsAsset* TargetIntersectPhysicsAsset = TargetIntersectOpPhysicsAsset;
+	TArray<int32> IKGoalMap;
+	TArray<FName> IKGoalNames;
+	TArray<int32> IKGoalCollisionIdx;
+	IKGoalMap.Reserve(CacheSourceBodyEffectVertIdx.Num());
+	IKGoalNames.Reserve(CacheSourceBodyEffectVertIdx.Num());
+	IKGoalCollisionIdx.Reserve(CacheSourceBodyEffectVertIdx.Num());
+	for (int EffectIdx = 0; EffectIdx < CacheSourceBodyEffectVertIdx.Num(); ++EffectIdx)
+	{
+		TArray<int32>& EffectIndices = CacheSourceBodyEffectVertIdx[EffectIdx];
+		if (EffectIndices.IsEmpty())
+		{
+			continue;
+		}
+
+		FName SourceBoneName = CacheSourceEffectBones[EffectIdx];
+		FName TargetBoneName = ApplyBodyMap(SourceBoneName);
+		if (!CacheTargetSkelIndices.Contains(TargetBoneName))
+		{
+			continue;
+		}
+		
+		FIKRigGoal* IKGoal = OutIKGoals.FindByPredicate([TargetBoneName](const FIKRigGoal& Other){return Other.BoneName == TargetBoneName;});
+		if (IKGoal)
+		{
+			int32 CheckCollisionIdx = TargetIntersectPhysicsAsset->FindBodyIndex(IKGoal->BoneName);
+			if (CheckCollisionIdx == INDEX_NONE)
+			{
+				continue;
+			}
+			IKGoalMap.Add(EffectIdx);
+			IKGoalNames.Add(SourceBoneName);
+			IKGoalCollisionIdx.Add(CheckCollisionIdx);
+			continue;
+		}
+		
+		int32 PinPropIdx = CachePinSourcePropBones.Find(SourceBoneName);
+		if (PinPropIdx != INDEX_NONE && CachePropId2PinBoneDataId[PinPropIdx] != INDEX_NONE)
+		{
+			if (!Settings.PinPropBones[CachePropId2PinBoneDataId[PinPropIdx]].PropSkeletalMeshAsset || !Settings.PinPropBones[CachePropId2PinBoneDataId[PinPropIdx]].PropSkeletalMeshAsset->GetPhysicsAsset())
+			{
+				continue;
+			}
+			IKGoalMap.Add(EffectIdx);
+			IKGoalNames.Add(SourceBoneName);
+			IKGoalCollisionIdx.Add(INDEX_NONE);
+			continue;
+		}
+	}
+	
+	TArray<TArray<int32>> GluedBodies;
+	TArray<FVector> GluedBodyDelta;
+	TArray<int32> IKGoal2GluedBodyIdx;
+	TArray<bool> IKGoalVisited;
+	TArray<int32> DfsStack;
+	DfsStack.Reserve(IKGoalMap.Num());
+	IKGoalVisited.Init(false, IKGoalMap.Num());
+	IKGoal2GluedBodyIdx.Init(INDEX_NONE, IKGoalMap.Num());
+	for (int32 IkGoalIdx0 = 0; IkGoalIdx0 < IKGoalMap.Num(); ++IkGoalIdx0)
+	{
+		DfsStack.Push(IkGoalIdx0);
+		while (!DfsStack.IsEmpty())
+		{
+			int32 IkGoalIdx = DfsStack.Pop();
+			if (IKGoalVisited[IkGoalIdx])
+			{
+				continue;
+			}
+			IKGoalVisited[IkGoalIdx] = true;
+			
+			int32& GluedBodyIdx = IKGoal2GluedBodyIdx[IkGoalIdx];
+			if (GluedBodyIdx == INDEX_NONE)
+			{
+				GluedBodies.AddDefaulted();
+				GluedBodyIdx = GluedBodies.Num() - 1;
+			}
+			TArray<int32>& GluedBody = GluedBodies[GluedBodyIdx];
+			GluedBody.Add(IkGoalIdx);
+			int32 EffectIdx = IKGoalMap[IkGoalIdx];
+			const TArray<int32>& EffectIndices = CacheSourceBodyEffectVertIdx[EffectIdx];
+			for (int Index : EffectIndices)
+			{
+				const FRelativeIKFrameBoneConstraint& BoneConstraint = FrameBoneConstraints[Index];
+				if (FMath::IsNearlyZero(BoneConstraint.Weight))
+				{
+					continue;
+				}
+
+				FName EffectBone = CachePairEffectBoneNames[Index];
+				int32 IkGoalIdx2 = IKGoalNames.Find(EffectBone);
+			
+				if (IkGoalIdx2 != INDEX_NONE)
+				{
+					if (BoneConstraint.SourceDist<Settings.GlueThreshold)
+					{
+						DfsStack.Push(IkGoalIdx2);
+						IKGoal2GluedBodyIdx[IkGoalIdx2] = GluedBodyIdx;
+					}
+				}
+			}
+		}
+	}
+	
+	GluedBodyDelta.SetNumZeroed(GluedBodies.Num());
+	for (int32 GlueBodyIdx=0; GlueBodyIdx<GluedBodies.Num(); ++GlueBodyIdx)
+	{
+		const TArray<int32>& GluedBody = GluedBodies[GlueBodyIdx];
+		FVector& IntersectDelta = GluedBodyDelta[GlueBodyIdx];
+		for (int32 IkGoalIdx: GluedBody)
+		{
+			int32 CollisionProxyIndex = IKGoalCollisionIdx[IkGoalIdx];
+			FName TargetBoneName = IKGoalNames[IkGoalIdx];
+			int32 TargetBoneIdx = CacheTargetSkelIndices[TargetBoneName];
+			const FTransform& TargetBonePose = TargetPose[TargetBoneIdx];
+			
+			FVector UpdateDelta = FVector::ZeroVector;
+			if (CollisionProxyIndex == INDEX_NONE) // Props
+			{
+				int32 PinPropIdx = CachePinSourcePropBones.Find(TargetBoneName);
+				FTransform CollisionProxyTfm = CachePinTargetPropTransforms[PinPropIdx];
+				CollisionProxyTfm.SetTranslation(CollisionProxyTfm.GetTranslation()+IntersectDelta);
+				const UPhysicsAsset* PropPhysicsAsset = Settings.PinPropBones[CachePropId2PinBoneDataId[PinPropIdx]].PropSkeletalMeshAsset->GetPhysicsAsset();
+				const FTransform PropAttachTransform = OffsetTransformsForSourceBones(TargetBoneName);
+				
+				FTransform UpdateTfm = CollisionProxyTfm;
+				for (int32 PropBodyIdx=0; PropBodyIdx<PropPhysicsAsset->SkeletalBodySetups.Num(); ++PropBodyIdx)
+				{
+					const FKAggregateGeom* CollisionProxyAggGeom = &PropPhysicsAsset->SkeletalBodySetups[PropBodyIdx]->AggGeom;
+					for (int j = 0; j < Settings.IntersectBodies.Num(); ++j)
+					{
+						FName BodyName = Settings.IntersectBodies[j];
+						int32 BodyPoseIdx = CacheTargetSkelIndices[BodyName];
+						if (BodyPoseIdx == INDEX_NONE)
+						{
+							continue;
+						}
+		
+						FTransform BodyTfm = TargetPose[BodyPoseIdx];
+						const FKAggregateGeom* BodyAggGeom = FPhysShapeUtils::FindBodyAggGeom(TargetIntersectPhysicsAsset, BodyName);
+						if (!BodyAggGeom)
+						{
+							continue;
+						}
+
+						FVector DeltaDir = FVector::ZeroVector;
+						double Dist = FBodyIntersectUtils::CalcIntersectionPairDelta(PropAttachTransform*UpdateTfm, CollisionProxyAggGeom, BodyTfm, BodyAggGeom, DeltaDir);
+						if (Dist > 0.0)
+						{
+							UpdateDelta += Dist * DeltaDir;
+							UpdateTfm.SetTranslation(CollisionProxyTfm.GetTranslation() + UpdateDelta);
+						}
+					}
+				}
+			}
+			else
+			{
+				FIKRigGoal* IKGoal = OutIKGoals.FindByPredicate([TargetBoneName](const FIKRigGoal& Other){return Other.BoneName == TargetBoneName;});
+				FVector CompGoalShapeLoc = GoalLocBlendCompSpace(IKGoal, TargetBonePose)+IntersectDelta;
+				FTransform CollisionProxyTfm = FTransform(TargetBonePose.Rotator(), CompGoalShapeLoc, TargetBonePose.GetScale3D());
+				const FKAggregateGeom* CollisionProxyAggGeom = FPhysShapeUtils::FindBodyAggGeom(TargetIntersectPhysicsAsset, TargetBoneName);
+				
+				FTransform UpdateTfm = CollisionProxyTfm;
+				for (int j = 0; j < Settings.IntersectBodies.Num(); ++j)
+				{
+					FName BodyName = Settings.IntersectBodies[j];
+					int32 BodyPoseIdx = CacheTargetSkelIndices[BodyName];
+					if (BodyPoseIdx == INDEX_NONE)
+					{
+						continue;
+					}
+		
+					int32 BodyPhysIdx = TargetIntersectPhysicsAsset->FindBodyIndex(BodyName);
+					if (BodyPhysIdx == INDEX_NONE)
+					{
+						continue;
+					}
+
+					// Allow CollisionProxyIndex = -1 to skip collision-enabled check (for prop capsules)
+					if (CollisionProxyIndex != INDEX_NONE && !TargetIntersectPhysicsAsset->IsCollisionEnabled(CollisionProxyIndex, BodyPhysIdx))
+					{
+						continue;
+					}
+		
+					FTransform BodyTfm = TargetPose[BodyPoseIdx];
+					const FKAggregateGeom* BodyAggGeom = FPhysShapeUtils::FindBodyAggGeom(TargetIntersectPhysicsAsset, BodyName);
+					if (!BodyAggGeom)
+					{
+						continue;
+					}
+
+					FVector DeltaDir = FVector::ZeroVector;
+					double Dist = FBodyIntersectUtils::CalcIntersectionPairDelta(UpdateTfm, CollisionProxyAggGeom, BodyTfm, BodyAggGeom, DeltaDir);
+					if (Dist > 0.0)
+					{
+						UpdateDelta += Dist * DeltaDir;
+						UpdateTfm.SetTranslation(CollisionProxyTfm.GetTranslation() + UpdateDelta);
+					}
+				}
+			}
+			IntersectDelta += UpdateDelta;
+		}
+	}
+	
+	bool bSetGoals = !Settings.bDebugDraw || !Settings.bDryRun;
+	if (bSetGoals)
+	{
+		for (int32 GlueBodyIdx=0; GlueBodyIdx<GluedBodies.Num(); ++GlueBodyIdx)
+		{
+			const TArray<int32>& GluedBody = GluedBodies[GlueBodyIdx];
+			const FVector& IntersectDelta = GluedBodyDelta[GlueBodyIdx]*Settings.IntersectOffsetScale;
+			for (int32 IkGoalIdx: GluedBody)
+			{
+				FName TargetBoneName = IKGoalNames[IkGoalIdx];
+				int32 TargetBoneIdx = CacheTargetSkelIndices[TargetBoneName];
+				const FTransform& TargetBonePose = TargetPose[TargetBoneIdx];
+				int32 CollisionProxyIndex = IKGoalCollisionIdx[IkGoalIdx];
+				if (CollisionProxyIndex == INDEX_NONE) // Props
+				{
+					int32 PinPropIdx = CachePinSourcePropBones.Find(TargetBoneName);
+					CachePinTargetPropTransforms[PinPropIdx].SetLocation(CachePinTargetPropTransforms[PinPropIdx].GetLocation()+IntersectDelta);
+				}
+				else
+				{
+					FIKRigGoal* IKGoal = OutIKGoals.FindByPredicate([TargetBoneName](const FIKRigGoal& Other){return Other.BoneName == TargetBoneName;});
+					FVector CompGoalShapeLoc = GoalLocBlendCompSpace(IKGoal, TargetBonePose)+IntersectDelta;
+		    	
+					IKGoal->bEnabled = true;
+					IKGoal->Position = CompGoalShapeLoc;
+					IKGoal->PositionAlpha = 1.0;
+					IKGoal->PositionSpace = EIKRigGoalSpace::Component;
+				}
+			}
+		}
+	}
+}
+
 void FIKRetargetRelativeIKOp::UpdateTargetIKGoal(
 	TArray<FIKRigGoal>& OutIKGoals,
 	FName TargetBoneName,
@@ -1668,10 +1942,10 @@ void FIKRetargetRelativeIKOp::ComputeTargetBodyTransform(FBodyTransform& OutTran
 {
 	FName TargetBoneName = ApplyBodyMap(SourceBoneName);
 	
-	FQuat BodyRot = GetBodyRotation(TargetPhysicsAsset, TargetBoneName);
-	FVector BodyTrans = GetBodyTranslation(TargetPhysicsAsset, TargetBoneName);
+	FQuat BodyRot = GetBodyRotation(TargetRelativeIKPhysicsAsset, TargetBoneName);
+	FVector BodyTrans = GetBodyTranslation(TargetRelativeIKPhysicsAsset, TargetBoneName);
 	// TODO: Probably should always use target body scale
-	FVector BodyScale = GetBodyOrientedScale(TargetPhysicsAsset, TargetBoneName);
+	FVector BodyScale = GetBodyOrientedScale(TargetRelativeIKPhysicsAsset, TargetBoneName);
 	
 	OutTransform.BoneToBody = GetTargetBoneDelta(TargetBoneName) * FTransform(BodyRot.Inverse());
 	OutTransform.BodyScale = BodyScale;
@@ -2502,11 +2776,11 @@ void FIKRetargetRelativeIKOp::DebugDraw(
 	
 	if (Settings.bDebugDrawPhysicsBodies)
 	{
-		if (Settings.TargetPhysicsAssetOverride)
+		if (Settings.TargetRelativeIKPhysicsAssetOverride)
 		{
 			// TODO: if only bake source transforms instead of PA, this source draw is meaningless
-			DebugDrawBodies(InPDI, InSourceTransform, 1.0, Settings.TargetPhysicsAssetOverride, DebugDrawInfo.SourceBodyInfo, true);
-			DebugDrawBodies(InPDI, InComponentTransform, InComponentScale, Settings.TargetPhysicsAssetOverride, DebugDrawInfo.TargetBodyInfo, false);
+			DebugDrawBodies(InPDI, InSourceTransform, 1.0, Settings.TargetRelativeIKPhysicsAssetOverride, DebugDrawInfo.SourceBodyInfo, true);
+			DebugDrawBodies(InPDI, InComponentTransform, InComponentScale, Settings.TargetRelativeIKPhysicsAssetOverride, DebugDrawInfo.TargetBodyInfo, false);
 		}
 	}
 	
diff --git a/Engine/Plugins/Experimental/Animation/RelativeIKOp/Source/RelativeIKOp/Public/RelativeIKOp.h b/Engine/Plugins/Experimental/Animation/RelativeIKOp/Source/RelativeIKOp/Public/RelativeIKOp.h
--- a/Engine/Plugins/Experimental/Animation/RelativeIKOp/Source/RelativeIKOp/Public/RelativeIKOp.h
+++ b/Engine/Plugins/Experimental/Animation/RelativeIKOp/Source/RelativeIKOp/Public/RelativeIKOp.h
@@ -3,6 +3,7 @@
 
 #include "Retargeter/IKRetargetOps.h"
 #include "Retargeter/RetargetOps/PinBoneOp.h"
+#include "Engine/SkeletalMesh.h"
 #include "Eigen/Dense"
 
 #include "RelativeIKOp.generated.h"
@@ -11,14 +12,17 @@
 
 #define LOCTEXT_NAMESPACE "RelativeIKOp"
 
+class BODYINTERSECTIKOP_API FPhysShapeUtils;
 struct FIKRigGoal;
 struct FKShapeElem;
+struct FKAggregateGeom;
 struct FResolvedBoneChain;
 struct FAnimMontageInstance;
 struct FIKRetargetRunIKRigOp;
 class UAnimInstance;
 class UAnimSequence;
 class UPhysicsAsset;
+class USkeletalMesh;
 class URelativeBodyBakeAnimNotify;
 class URelativePropsBakeAnimNotify;
 
@@ -122,6 +126,10 @@ struct FPinBoneSettings
 	// Alpha between primary and secondary pair distance relationship
 	UPROPERTY(EditAnywhere, BlueprintReadWrite, Category = "Pinning", meta=(ReinitializeOnEdit))
 	double PropScalar = 1.0;
+	
+	// Prop Skeletal Mesh
+	UPROPERTY(EditAnywhere, BlueprintReadWrite, Category = "Pinning", meta=(ReinitializeOnEdit))
+	TObjectPtr<USkeletalMesh> PropSkeletalMeshAsset;
 };
 
 UENUM()
@@ -171,8 +179,8 @@ struct FIKRetargetRelativeIKOpSettings : public FIKRetargetOpSettingsBase
 
 	// Target mesh physics asset for retargeting relative body pair vertices
 	UPROPERTY(EditAnywhere, BlueprintReadWrite, Category = "Physics", meta=(ReinitializeOnEdit))
-	TObjectPtr<UPhysicsAsset> TargetPhysicsAssetOverride;
-
+	TObjectPtr<UPhysicsAsset> TargetRelativeIKPhysicsAssetOverride;
+	
 	// Physics body source -> target name mapping 
 	UPROPERTY(EditAnywhere, BlueprintReadWrite, Category = "Source Target Mapping", meta=(ReinitializeOnEdit))
 	TMap<FName,FName> BodyMapping;
@@ -310,6 +318,26 @@ struct FIKRetargetRelativeIKOpSettings : public FIKRetargetOpSettingsBase
 	UPROPERTY(EditAnywhere, BlueprintReadWrite, Category = "Prop Scaling", meta=(ClampMin = "0.001"))
 	double PropScalarDamping = 500;
 
+	// Enable Collison push out for Ikgoals and Props
+	UPROPERTY(EditAnywhere, BlueprintReadWrite, Category = "Intersection")
+	bool bEnableIntersectPushOut = true;
+	
+	// Target mesh physics asset for Intersect push out
+	UPROPERTY(EditAnywhere, BlueprintReadWrite, Category = "Intersection", meta=(ReinitializeOnEdit))
+	TObjectPtr<UPhysicsAsset> TargetIntersectOpPhysicsAssetOverride;
+	
+	// Max distance to glue bodies (very close to a body)
+	UPROPERTY(EditAnywhere, BlueprintReadWrite, Category = "Intersection", meta=(ClampMin = "0.01"))
+	double GlueThreshold = 5;
+	
+	// Max distance to glue bodies (very close to a body)
+	UPROPERTY(EditAnywhere, BlueprintReadWrite, Category = "Intersection", meta=(ClampMin = "0.01"))
+	double IntersectOffsetScale = 1;
+	
+	// Bone names with attached physics bodies to do trivial intersection against
+	// **NOTE: Only Sphere, Capsule and Box bodies currently supported**
+	UPROPERTY(EditAnywhere, BlueprintReadWrite, Category = "Intersection")
+	TArray<FName> IntersectBodies;
 	
 	// Distance blend body pairs towards each contact position if both can move
 	UPROPERTY(EditAnywhere, BlueprintReadWrite, Category = "Test Features")
@@ -397,6 +425,7 @@ private:
 	};
 	
 	void InitializePinData(const FIKRetargetProcessor& Processor);
+	void InitializeIntersectionData(const FIKRetargetProcessor& Processor);
 	
 	void CachePinBoneTransforms(
 		const FIKRetargetProcessor& Processor,
@@ -444,6 +473,11 @@ private:
 		FDebugRelativeIKDrawInfo& DebugInfo,
 		TArray<FDebugRelativeTargetPairSpace>& DebugPairSpaces);
 	
+	void ApplyIntersectPushOut(
+		TArray<FIKRigGoal>& OutIKGoals,
+		const TArray<FTransform>& TargetPose,
+		FDebugRelativeIKDrawInfo& DebugInfo);
+	
 	void ApplyMatrixBodyPairTargetsWithPropScale(
 		TArray<FIKRigGoal>& OutIKGoals,
 		const TArray<FTransform>& TargetPose,
@@ -556,7 +590,8 @@ private:
 private:
 	
 	// TObjectPtr<const UPhysicsAsset> SourcePhysicsAsset;
-	TObjectPtr<const UPhysicsAsset> TargetPhysicsAsset;
+	TObjectPtr<const UPhysicsAsset> TargetRelativeIKPhysicsAsset;
+	TObjectPtr<const UPhysicsAsset> TargetIntersectOpPhysicsAsset;
 	
 	TObjectPtr<const UIKRigDefinition> TargetIKRig;
 	
diff --git a/Engine/Plugins/Experimental/Animation/RelativeIKOp/Source/RelativeIKOp/RelativeIKOp.Build.cs b/Engine/Plugins/Experimental/Animation/RelativeIKOp/Source/RelativeIKOp/RelativeIKOp.Build.cs
--- a/Engine/Plugins/Experimental/Animation/RelativeIKOp/Source/RelativeIKOp/RelativeIKOp.Build.cs
+++ b/Engine/Plugins/Experimental/Animation/RelativeIKOp/Source/RelativeIKOp/RelativeIKOp.Build.cs
@@ -49,6 +49,7 @@ namespace UnrealBuildTool.Rules
 				"Renderer",
 				"IKRig",
 				"RelativeBodyAnimInfo",
+				"BodyIntersectIKOp",
 				"Eigen"
 				}
 				);
