RelativeIK: Prop push out

#Synced-CL 52491934

diff --git a/Engine/Plugins/Experimental/Animation/RelativeIKOp/Source/BodyIntersectIKOp/Private/PhysicsBodyHelpers.cpp b/Engine/Plugins/Experimental/Animation/RelativeIKOp/Source/BodyIntersectIKOp/Private/PhysicsBodyHelpers.cpp
--- a/Engine/Plugins/Experimental/Animation/RelativeIKOp/Source/BodyIntersectIKOp/Private/PhysicsBodyHelpers.cpp
+++ b/Engine/Plugins/Experimental/Animation/RelativeIKOp/Source/BodyIntersectIKOp/Private/PhysicsBodyHelpers.cpp
@@ -468,6 +468,12 @@ double FPhysShapeUtils::CalcShapeSmallRadius(const FKShapeElem* ShapeElem)
 
 double FPhysShapeUtils::CalcShapeAvgRadius(const FKShapeElem* ShapeElem)
 {
+	const double DefaultRadius = 0.0;
+	if (!ShapeElem)
+	{
+		return DefaultRadius;
+	}
+	
 	switch (ShapeElem->GetShapeType())
 	{
 	case EAggCollisionShape::Sphere:
diff --git a/Engine/Plugins/Experimental/Animation/RelativeIKOp/Source/RelativeIKOp/Private/RelativeIKOp.cpp b/Engine/Plugins/Experimental/Animation/RelativeIKOp/Source/RelativeIKOp/Private/RelativeIKOp.cpp
--- a/Engine/Plugins/Experimental/Animation/RelativeIKOp/Source/RelativeIKOp/Private/RelativeIKOp.cpp
+++ b/Engine/Plugins/Experimental/Animation/RelativeIKOp/Source/RelativeIKOp/Private/RelativeIKOp.cpp
@@ -15,6 +15,11 @@
 #include "Retargeter/RetargetOps/RunIKRigOp.h"
 #include "Eigen/SVD"
 
+#include "AnimationRuntime.h"
+#include "BoneContainer.h"
+#include "BoneIndices.h"
+#include "BonePose.h"
+
 #if WITH_EDITOR
 #include "PrimitiveDrawingUtils.h"
 #include "Engine/EngineTypes.h"
@@ -92,19 +97,19 @@ bool FIKRetargetRelativeIKOp::Initialize(
 		return false;
 	}
 	
-	TargetRelativeIKPhysicsAsset = Settings.TargetRelativeIKPhysicsAssetOverride;
+	TargetRelativeIKPhysicsAsset = Settings.TargetPhysicsAssetOverride;
 	TargetIntersectOpPhysicsAsset= Settings.TargetIntersectOpPhysicsAssetOverride;
 	if (!TargetRelativeIKPhysicsAsset)
 	{
 		InLog.LogWarning( FText::Format(
-		LOCTEXT("MissingPhysicsAssets", "{0}: Both Source and Target Physics Assets must be specified. "), FText::FromName(GetName())));
+		LOCTEXT("MissingTargetPhysicsAssets", "{0}: Target RelativeIK Physics Assets must be specified. "), FText::FromName(GetName())));
 		return false;
 	}
 	if (Settings.bEnableIntersectPushOut && !TargetIntersectOpPhysicsAsset)
 	{
 		TargetIntersectOpPhysicsAsset = TargetRelativeIKPhysicsAsset;
 		InLog.LogWarning( FText::Format(
-		LOCTEXT("MissingPhysicsAssets", "Missing TargetIntersectOpPhysicsAsset Override, use TargetRelativeIKPhysicsAsset Override instead!"), FText::FromName(GetName())));
+		LOCTEXT("MissingTargetIntersectOpPhysicsAssets", "{0}: Missing TargetIntersectOpPhysicsAsset Override, use TargetRelativeIKPhysicsAsset Override instead!"), FText::FromName(GetName())));
 	}
 
 	UpdateCacheBoneChains(InProcessor, ParentRigOp);
@@ -1433,10 +1438,10 @@ void FIKRetargetRelativeIKOp::ApplyIntersectPushOut(
 {
 	const UPhysicsAsset* TargetIntersectPhysicsAsset = TargetIntersectOpPhysicsAsset;
 	TArray<int32> IKGoalMap;
-	TArray<FName> IKGoalNames;
+	TArray<FName> IKGoalSourceNames;
 	TArray<int32> IKGoalCollisionIdx;
 	IKGoalMap.Reserve(CacheSourceBodyEffectVertIdx.Num());
-	IKGoalNames.Reserve(CacheSourceBodyEffectVertIdx.Num());
+	IKGoalSourceNames.Reserve(CacheSourceBodyEffectVertIdx.Num());
 	IKGoalCollisionIdx.Reserve(CacheSourceBodyEffectVertIdx.Num());
 	for (int EffectIdx = 0; EffectIdx < CacheSourceBodyEffectVertIdx.Num(); ++EffectIdx)
 	{
@@ -1462,7 +1467,7 @@ void FIKRetargetRelativeIKOp::ApplyIntersectPushOut(
 				continue;
 			}
 			IKGoalMap.Add(EffectIdx);
-			IKGoalNames.Add(SourceBoneName);
+			IKGoalSourceNames.Add(SourceBoneName);
 			IKGoalCollisionIdx.Add(CheckCollisionIdx);
 			continue;
 		}
@@ -1475,7 +1480,7 @@ void FIKRetargetRelativeIKOp::ApplyIntersectPushOut(
 				continue;
 			}
 			IKGoalMap.Add(EffectIdx);
-			IKGoalNames.Add(SourceBoneName);
+			IKGoalSourceNames.Add(SourceBoneName);
 			IKGoalCollisionIdx.Add(INDEX_NONE);
 			continue;
 		}
@@ -1520,7 +1525,7 @@ void FIKRetargetRelativeIKOp::ApplyIntersectPushOut(
 				}
 
 				FName EffectBone = CachePairEffectBoneNames[Index];
-				int32 IkGoalIdx2 = IKGoalNames.Find(EffectBone);
+				int32 IkGoalIdx2 = IKGoalSourceNames.Find(EffectBone);
 			
 				if (IkGoalIdx2 != INDEX_NONE)
 				{
@@ -1542,45 +1547,105 @@ void FIKRetargetRelativeIKOp::ApplyIntersectPushOut(
 		for (int32 IkGoalIdx: GluedBody)
 		{
 			int32 CollisionProxyIndex = IKGoalCollisionIdx[IkGoalIdx];
-			FName TargetBoneName = IKGoalNames[IkGoalIdx];
+			FName SourceBoneName = IKGoalSourceNames[IkGoalIdx];
+			FName TargetBoneName = ApplyBodyMap(SourceBoneName);
 			int32 TargetBoneIdx = CacheTargetSkelIndices[TargetBoneName];
 			const FTransform& TargetBonePose = TargetPose[TargetBoneIdx];
 			
 			FVector UpdateDelta = FVector::ZeroVector;
 			if (CollisionProxyIndex == INDEX_NONE) // Props
 			{
-				int32 PinPropIdx = CachePinSourcePropBones.Find(TargetBoneName);
+				int32 PinPropIdx = CachePinSourcePropBones.Find(SourceBoneName);
 				FTransform CollisionProxyTfm = CachePinTargetPropTransforms[PinPropIdx];
 				CollisionProxyTfm.SetTranslation(CollisionProxyTfm.GetTranslation()+IntersectDelta);
 				const UPhysicsAsset* PropPhysicsAsset = Settings.PinPropBones[CachePropId2PinBoneDataId[PinPropIdx]].PropSkeletalMeshAsset->GetPhysicsAsset();
-				const FTransform PropAttachTransform = OffsetTransformsForSourceBones(TargetBoneName);
-				
-				FTransform UpdateTfm = CollisionProxyTfm;
-				for (int32 PropBodyIdx=0; PropBodyIdx<PropPhysicsAsset->SkeletalBodySetups.Num(); ++PropBodyIdx)
+				const FTransform PropAttachTransform = OffsetTransformsForSourceBones(SourceBoneName);
+				const UAnimSequence* PropAnimation = Settings.PinPropBones[CachePropId2PinBoneDataId[PinPropIdx]].PropAnimSequence;
+				if (PropAnimation)
 				{
-					const FKAggregateGeom* CollisionProxyAggGeom = &PropPhysicsAsset->SkeletalBodySetups[PropBodyIdx]->AggGeom;
-					for (int j = 0; j < Settings.IntersectBodies.Num(); ++j)
+					float TestPropAnimLength = PropAnimation->GetPlayLength();
+					float TestPropPlayhead = FMath::Max(FMath::Fmod(AnimSeqPlayHead, TestPropAnimLength), 0.0f);
+					if (TestPropPlayhead<0.f)
 					{
-						FName BodyName = Settings.IntersectBodies[j];
-						int32 BodyPoseIdx = CacheTargetSkelIndices[BodyName];
-						if (BodyPoseIdx == INDEX_NONE)
+						continue;
+					}
+					TArray<FName> OutPropBones;
+					TArray<FTransform> OutPropPose;
+					GetAnimSeqFramePose(PropAnimation, TestPropPlayhead, OutPropBones, OutPropPose);
+					FTransform UpdateTfm = CollisionProxyTfm;
+					for (int32 PropBodyIdx=0; PropBodyIdx<PropPhysicsAsset->SkeletalBodySetups.Num(); ++PropBodyIdx)
+					{
+						const FKAggregateGeom* CollisionProxyAggGeom = &PropPhysicsAsset->SkeletalBodySetups[PropBodyIdx]->AggGeom;
+						const FName PropBoneName = PropPhysicsAsset->SkeletalBodySetups[PropBodyIdx]->BoneName;
+						int32 OutPropBonesIdx = OutPropBones.Find(PropBoneName);
+						if (OutPropBonesIdx == INDEX_NONE)
 						{
 							continue;
 						}
-		
-						FTransform BodyTfm = TargetPose[BodyPoseIdx];
-						const FKAggregateGeom* BodyAggGeom = FPhysShapeUtils::FindBodyAggGeom(TargetIntersectPhysicsAsset, BodyName);
-						if (!BodyAggGeom)
+						FTransform PropBoneTfm = OutPropPose[OutPropBonesIdx];
+						for (int j = 0; j < Settings.IntersectBodies.Num(); ++j)
 						{
-							continue;
+							FName BodyName = Settings.IntersectBodies[j];
+							if (!CacheTargetSkelIndices.Contains(BodyName))
+							{
+								continue;
+							}
+							int32 BodyPoseIdx = CacheTargetSkelIndices[BodyName];
+							if (BodyPoseIdx == INDEX_NONE)
+							{
+								continue;
+							}
+		
+							FTransform BodyTfm = TargetPose[BodyPoseIdx];
+							const FKAggregateGeom* BodyAggGeom = FPhysShapeUtils::FindBodyAggGeom(TargetIntersectPhysicsAsset, BodyName);
+							if (!BodyAggGeom)
+							{
+								continue;
+							}
+
+							FVector DeltaDir = FVector::ZeroVector;
+							double Dist = FBodyIntersectUtils::CalcIntersectionPairDelta(PropAttachTransform*PropBoneTfm*UpdateTfm, CollisionProxyAggGeom, BodyTfm, BodyAggGeom, DeltaDir);
+							if (Dist > 0.0)
+							{
+								UpdateDelta += Dist * DeltaDir;
+								UpdateTfm.SetTranslation(CollisionProxyTfm.GetTranslation() + UpdateDelta);
+							}
 						}
-
-						FVector DeltaDir = FVector::ZeroVector;
-						double Dist = FBodyIntersectUtils::CalcIntersectionPairDelta(PropAttachTransform*UpdateTfm, CollisionProxyAggGeom, BodyTfm, BodyAggGeom, DeltaDir);
-						if (Dist > 0.0)
+					}
+				}
+				else
+				{
+					FTransform UpdateTfm = CollisionProxyTfm;
+					for (int32 PropBodyIdx=0; PropBodyIdx<PropPhysicsAsset->SkeletalBodySetups.Num(); ++PropBodyIdx)
+					{
+						const FKAggregateGeom* CollisionProxyAggGeom = &PropPhysicsAsset->SkeletalBodySetups[PropBodyIdx]->AggGeom;
+						for (int j = 0; j < Settings.IntersectBodies.Num(); ++j)
 						{
-							UpdateDelta += Dist * DeltaDir;
-							UpdateTfm.SetTranslation(CollisionProxyTfm.GetTranslation() + UpdateDelta);
+							FName BodyName = Settings.IntersectBodies[j];
+							if (!CacheTargetSkelIndices.Contains(BodyName))
+							{
+								continue;
+							}
+							int32 BodyPoseIdx = CacheTargetSkelIndices[BodyName];
+							if (BodyPoseIdx == INDEX_NONE)
+							{
+								continue;
+							}
+		
+							FTransform BodyTfm = TargetPose[BodyPoseIdx];
+							const FKAggregateGeom* BodyAggGeom = FPhysShapeUtils::FindBodyAggGeom(TargetIntersectPhysicsAsset, BodyName);
+							if (!BodyAggGeom)
+							{
+								continue;
+							}
+
+							FVector DeltaDir = FVector::ZeroVector;
+							double Dist = FBodyIntersectUtils::CalcIntersectionPairDelta(PropAttachTransform*UpdateTfm, CollisionProxyAggGeom, BodyTfm, BodyAggGeom, DeltaDir);
+							if (Dist > 0.0)
+							{
+								UpdateDelta += Dist * DeltaDir;
+								UpdateTfm.SetTranslation(CollisionProxyTfm.GetTranslation() + UpdateDelta);
+							}
 						}
 					}
 				}
@@ -1596,6 +1661,10 @@ void FIKRetargetRelativeIKOp::ApplyIntersectPushOut(
 				for (int j = 0; j < Settings.IntersectBodies.Num(); ++j)
 				{
 					FName BodyName = Settings.IntersectBodies[j];
+					if (!CacheTargetSkelIndices.Contains(BodyName))
+					{
+						continue;
+					}
 					int32 BodyPoseIdx = CacheTargetSkelIndices[BodyName];
 					if (BodyPoseIdx == INDEX_NONE)
 					{
@@ -1643,13 +1712,14 @@ void FIKRetargetRelativeIKOp::ApplyIntersectPushOut(
 			const FVector& IntersectDelta = GluedBodyDelta[GlueBodyIdx]*Settings.IntersectOffsetScale;
 			for (int32 IkGoalIdx: GluedBody)
 			{
-				FName TargetBoneName = IKGoalNames[IkGoalIdx];
+				FName SourceBoneName =IKGoalSourceNames[IkGoalIdx];
+				FName TargetBoneName = ApplyBodyMap(SourceBoneName);
 				int32 TargetBoneIdx = CacheTargetSkelIndices[TargetBoneName];
 				const FTransform& TargetBonePose = TargetPose[TargetBoneIdx];
 				int32 CollisionProxyIndex = IKGoalCollisionIdx[IkGoalIdx];
 				if (CollisionProxyIndex == INDEX_NONE) // Props
 				{
-					int32 PinPropIdx = CachePinSourcePropBones.Find(TargetBoneName);
+					int32 PinPropIdx = CachePinSourcePropBones.Find(SourceBoneName);
 					CachePinTargetPropTransforms[PinPropIdx].SetLocation(CachePinTargetPropTransforms[PinPropIdx].GetLocation()+IntersectDelta);
 				}
 				else
@@ -2742,6 +2812,56 @@ void FIKRetargetRelativeIKOp::OnParentReinitPropertyEdited(const FIKRetargetOpBa
 {
 }
 
+void FIKRetargetRelativeIKOp::GetAnimSeqFramePose(const UAnimSequence* AnimSeq, double Time, TArray<FName>& OutBones, TArray<FTransform>& OutPose) const
+{
+	const FReferenceSkeleton& Skel = AnimSeq->GetSkeleton()->GetReferenceSkeleton();
+	int32 NumBones = Skel.GetNum();
+	
+	OutBones.Reset(NumBones);
+	OutPose.Reset(NumBones);
+
+	// Init bone container for pulling animseq data
+	TArray<FBoneIndexType> BoneIndices;
+	BoneIndices.Reserve(NumBones);
+	for (FBoneIndexType BoneIndex = 0; BoneIndex < NumBones; ++BoneIndex)
+	{
+		BoneIndices.Add(BoneIndex);
+		OutBones.Add(Skel.GetBoneName(BoneIndex));
+	}
+	
+	FBoneContainer BoneContainer;
+	BoneContainer.InitializeTo(BoneIndices, UE::Anim::FCurveFilterSettings(UE::Anim::ECurveFilterMode::DisallowAll), *AnimSeq->GetSkeleton());
+	BoneContainer.SetUseRAWData(false);
+	
+	// Setup bone container
+	FCompactPose CompactPose;
+	FBlendedCurve Curves;
+	UE::Anim::FStackAttributeContainer Attributes;
+	FAnimationPoseData FramePoseData(CompactPose, Curves, Attributes);
+	
+	CompactPose.SetBoneContainer(&BoneContainer);
+	Curves.InitFrom(BoneContainer);
+	
+	FAnimExtractContext FrameExtractionCtx(Time, false);
+	AnimSeq->GetAnimationPose(FramePoseData, FrameExtractionCtx);
+	
+	FramePoseData.GetPose().CopyBonesTo(OutPose);
+	for (int32 BoneIndex = 0; BoneIndex < NumBones; ++BoneIndex)
+	{
+		const int32 ParentBoneIndex = Skel.GetParentIndex(BoneIndex);
+
+		FTransform ParentTransform = FTransform::Identity;
+		if (ParentBoneIndex != INDEX_NONE)
+		{
+			// TODO: Double-check parents always befor children in list
+			FString ParentBoneName = Skel.GetBoneName(ParentBoneIndex).ToString();
+			ParentTransform = OutPose[ParentBoneIndex];
+		}
+		
+		OutPose[BoneIndex] = OutPose[BoneIndex] * ParentTransform;
+	}
+}
+
 #if WITH_EDITOR
 
 void FIKRetargetRelativeIKOp::UpdateFromAnimSequence(UAnimSequence* Sequence, float FrameTime)
@@ -2750,6 +2870,7 @@ void FIKRetargetRelativeIKOp::UpdateFromAnimSequence(UAnimSequence* Sequence, fl
 	UpdateAnimSeqData(Sequence);
 }
 
+
 FCriticalSection FIKRetargetRelativeIKOp::DebugDataMutex;
 
 void FIKRetargetRelativeIKOp::DebugDraw(
@@ -2776,11 +2897,11 @@ void FIKRetargetRelativeIKOp::DebugDraw(
 	
 	if (Settings.bDebugDrawPhysicsBodies)
 	{
-		if (Settings.TargetRelativeIKPhysicsAssetOverride)
+		if (Settings.TargetPhysicsAssetOverride)
 		{
 			// TODO: if only bake source transforms instead of PA, this source draw is meaningless
-			DebugDrawBodies(InPDI, InSourceTransform, 1.0, Settings.TargetRelativeIKPhysicsAssetOverride, DebugDrawInfo.SourceBodyInfo, true);
-			DebugDrawBodies(InPDI, InComponentTransform, InComponentScale, Settings.TargetRelativeIKPhysicsAssetOverride, DebugDrawInfo.TargetBodyInfo, false);
+			DebugDrawBodies(InPDI, InSourceTransform, 1.0, Settings.TargetPhysicsAssetOverride, DebugDrawInfo.SourceBodyInfo, true);
+			DebugDrawBodies(InPDI, InComponentTransform, InComponentScale, Settings.TargetPhysicsAssetOverride, DebugDrawInfo.TargetBodyInfo, false);
 		}
 	}
 	
diff --git a/Engine/Plugins/Experimental/Animation/RelativeIKOp/Source/RelativeIKOp/Public/RelativeIKOp.h b/Engine/Plugins/Experimental/Animation/RelativeIKOp/Source/RelativeIKOp/Public/RelativeIKOp.h
--- a/Engine/Plugins/Experimental/Animation/RelativeIKOp/Source/RelativeIKOp/Public/RelativeIKOp.h
+++ b/Engine/Plugins/Experimental/Animation/RelativeIKOp/Source/RelativeIKOp/Public/RelativeIKOp.h
@@ -128,8 +128,12 @@ struct FPinBoneSettings
 	double PropScalar = 1.0;
 	
 	// Prop Skeletal Mesh
-	UPROPERTY(EditAnywhere, BlueprintReadWrite, Category = "Pinning", meta=(ReinitializeOnEdit))
+	UPROPERTY(EditAnywhere, BlueprintReadWrite, Category = "Intersection", meta=(ReinitializeOnEdit))
 	TObjectPtr<USkeletalMesh> PropSkeletalMeshAsset;
+	
+	// Prop mesh to anim sequence
+	UPROPERTY(EditAnywhere, BlueprintReadWrite, Category = "Intersection", meta=(ReinitializeOnEdit))
+	TObjectPtr<UAnimSequence> PropAnimSequence;
 };
 
 UENUM()
@@ -179,7 +183,7 @@ struct FIKRetargetRelativeIKOpSettings : public FIKRetargetOpSettingsBase
 
 	// Target mesh physics asset for retargeting relative body pair vertices
 	UPROPERTY(EditAnywhere, BlueprintReadWrite, Category = "Physics", meta=(ReinitializeOnEdit))
-	TObjectPtr<UPhysicsAsset> TargetRelativeIKPhysicsAssetOverride;
+	TObjectPtr<UPhysicsAsset> TargetPhysicsAssetOverride;
 	
 	// Physics body source -> target name mapping 
 	UPROPERTY(EditAnywhere, BlueprintReadWrite, Category = "Source Target Mapping", meta=(ReinitializeOnEdit))
@@ -557,6 +561,7 @@ private:
 	static FKShapeElem* FindBodyShape(const UPhysicsAsset* PhysAsset, FName BoneName);
 	
 	static double GetDistanceWeight(EDistanceWeightMode DistanceWeightMode, double Distance, double DistThreshold, double DistHalfPointLambda = 0.5);
+	void GetAnimSeqFramePose(const UAnimSequence* AnimSeq, double Time, TArray<FName>& OutBones, TArray<FTransform>& OutPose) const;
 
 #if WITH_EDITOR
 public:
